Promote experimental dynamic wind plugin into licensee accessible location.

#virtualized

diff --git a/Engine/Plugins/Experimental/DynamicWind/DynamicWind.uplugin b/Engine/Plugins/Experimental/DynamicWind/DynamicWind.uplugin
new file mode 100644
--- /dev/null
+++ b/Engine/Plugins/Experimental/DynamicWind/DynamicWind.uplugin
@@ -0,0 +1,28 @@
+{
+	"FileVersion": 1,
+	"Version": 1,
+	"VersionName": "0.1",
+	"FriendlyName": "Dynamic Wind",
+	"Description": "Extremely experimental dynamic wind support for Nanite foliage.",
+	"Category": "Rendering",
+	"CreatedBy": "Epic Games, Inc.",
+	"CreatedByURL": "https://epicgames.com",
+	"MarketplaceURL": "",
+	"SupportURL": "",
+	"EnabledByDefault": false,
+	"CanContainContent": true,
+	"IsExperimentalVersion": true,
+	"Installed": false,
+	"Modules": [
+		{
+			"Name": "DynamicWind",
+			"Type": "Runtime",
+			"LoadingPhase": "PostConfigInit"
+		},
+		{
+			"Name": "DynamicWindEditor",
+			"Type": "Editor",
+			"LoadingPhase": "Default"
+		}
+	]
+}
\ No newline at end of file
diff --git a/Engine/Plugins/Experimental/DynamicWind/Shaders/DynamicWindEval.usf b/Engine/Plugins/Experimental/DynamicWind/Shaders/DynamicWindEval.usf
new file mode 100644
--- /dev/null
+++ b/Engine/Plugins/Experimental/DynamicWind/Shaders/DynamicWindEval.usf
@@ -0,0 +1,257 @@
+// Copyright Epic Games, Inc. All Rights Reserved.
+
+#include "/Engine/Private/Common.ush"
+#include "/Engine/Private/BoneTransform.ush"
+#include "/Engine/Shared/SkinningDefinitions.h"
+#include "/Engine/Private/ComputeShaderUtils.ush"
+#include "/Engine/Private/Nanite/NaniteDataDecode.ush"
+
+#include "Shared/DynamicWindCommon.ush"
+
+// Enable to debug always outputting identity bones
+#define FORCE_IDENTITY_OUTPUT 0
+
+StructuredBuffer<FDynamicWindSkeletonData> SkeletonBuffer;
+StructuredBuffer<FDynamicWindBlockHeader> HeaderBuffer;
+ByteAddressBuffer BoneDataBuffer;
+RWByteAddressBuffer TransformBuffer;
+Texture2D<float4> WindTexture;
+SamplerState WindSampler;
+
+float Time;
+float DeltaTime;
+float WindSpeed;
+float WindAmplitude;
+float3 WindDirection;
+float4 Debug;
+
+// === WIND PARAMETERS ===============================================================================================================================
+
+// Wind speed and amplitude modulations
+#define DYNAMIC_WIND_SPEED_MOD 1.25f
+#define DYNAMIC_WIND_AMPLITUDE_MOD 3.0f
+
+// Flat value added to the final amplitude
+#define DYNAMIC_WIND_BEND_OFFSET 0.1f
+
+// Attenuation modifier based on the wind direction
+#define DYNAMIC_WIND_ALIGN_ATTENUATION 0.5f
+
+// SPECIFIC FOR BRANCHES
+
+// Attenuation modifier based on the bone height (0-1 normalized on skeleton height)
+#define DYNAMIC_WIND_HEIGHT_ATTENUATION 0.1f
+
+// Modulation of the angle of the torsion rotation of a branch
+#define DYNAMIC_WIND_TORSION_MOD 0.5f
+
+// Attenuation by branch length
+#define DYNAMIC_WIND_LONG_BRANCH_SPEED_MOD 0.5f
+#define DYNAMIC_WIND_LONG_BRANCH_MAX 500.0f
+#define DYNAMIC_WIND_LONG_BRANCH_MIN 100.0f
+
+// Gust speed
+#define DYNAMIC_WIND_GUST_SPEED_MIN 0.05f
+#define DYNAMIC_WIND_GUST_SPEED_MAX 0.1f
+
+// TODO Use a cvar instead of a define
+#define DYNAMIC_WIND_DIRECTIONALITY_SLICES 8.0f
+
+uint GetBoneDataAddress(uint BoneDataOffset, uint BoneIndex)
+{
+	return (BoneDataOffset + BoneIndex) * uint(sizeof(FDynamicWindBoneData));
+}
+
+float RemapValue(float Value, float InLow, float InHigh, float OutLow, float OutHigh)
+{
+	return OutLow + (Value - InLow) * (OutHigh - OutLow) / (InHigh - InLow);
+}
+
+FQuat QuatAttenuate(FQuat Q1, half Attenuation)
+{
+	const FQuat Q0 = QuatIdentity();
+	const half C = dot(Q0, Q1);
+	return normalize(lerp(Q0, C < 0.0f ? -Q1 : Q1, Attenuation));
+}
+
+float WindFlow(const float RandomValue, const float Distance = 0.0f, const float BlendByGroup = 0.0f, const float ChainLength = 0.0f, float TrunkBlend = 0.0f)
+{	
+	// Reducing the speed on lower branches to simulate wood thickness and weight
+	const float RemappedBranchLength = RemapValue(ChainLength, DYNAMIC_WIND_LONG_BRANCH_MAX, DYNAMIC_WIND_LONG_BRANCH_MIN, DYNAMIC_WIND_LONG_BRANCH_SPEED_MOD, 1.0f);
+	const float SpeedByLengthModulation = max(TrunkBlend, clamp(RemappedBranchLength, DYNAMIC_WIND_LONG_BRANCH_SPEED_MOD, 1.0f));
+	
+	// 100.0f is a manual modulation to keep the speed inside human-readable range
+	const float BiasedWindSpeed = WindSpeed / 100.0f;
+	const float FlowSpeed = Time * BiasedWindSpeed * SpeedByLengthModulation * DYNAMIC_WIND_SPEED_MOD;
+	const float2 UV = float2(RandomValue, frac(Distance * RandomValue)) - float2(frac(FlowSpeed), 0.0f);
+	
+	const float4 Sample = WindTexture.SampleLevel(WindSampler, UV, 0);
+
+	// Fractal Brownian noise approximation
+	float FinalIntensity = Sample.x + lerp(0.0f, Sample.w * 0.5f, BlendByGroup);
+	FinalIntensity = (DYNAMIC_WIND_BEND_OFFSET + FinalIntensity) * DYNAMIC_WIND_AMPLITUDE_MOD * WindSpeed / 100.0f;
+
+	return FinalIntensity;
+}
+
+float WindWave(float T)
+{
+	return cos(PI * T) * cos(3.0f * PI * T) * cos(5.0f * PI * T) * cos(7.0f * PI * T);
+}
+
+float RampStrength(float X)
+{
+	return X / (X * X + 1.0f);
+}
+
+[numthreads(BONES_PER_GROUP, 1, 1)]
+void WindEvalCS(uint3 GroupId : SV_GroupID, uint GroupThreadIndex : SV_GroupIndex)
+{
+	const uint GlobalBlockIndex	= GroupId.x;
+	const uint BoneIndexInBlock	= GroupThreadIndex;
+
+	const FDynamicWindBlockHeader Header = HeaderBuffer[GlobalBlockIndex];
+	if (BoneIndexInBlock >= Header.BlockNumBones)
+	{
+		return;
+	}
+
+	if (Header.SkeletonDataIndex == ~uint(0) || FORCE_IDENTITY_OUTPUT)
+	{
+		const float4x3 Identity =
+		{
+			float3(1, 0, 0),
+			float3(0, 1, 0),
+			float3(0, 0, 1),
+			float3(0, 0, 0)
+		};
+		StoreCompressedBoneTransform(TransformBuffer, Header.BlockDstTransformOffset, BoneIndexInBlock, Identity);
+		StoreCompressedBoneTransform(TransformBuffer, Header.BlockDstTransformOffset, BoneIndexInBlock + Header.TotalTransformCount, Identity);
+
+		return;
+	}
+
+	FDynamicWindSkeletonData SkeletonData = SkeletonBuffer[Header.SkeletonDataIndex];
+
+	const uint SkeletonBoneIndex = Header.BlockBoneOffset + BoneIndexInBlock;
+	const uint BoneDataAddress = GetBoneDataAddress(SkeletonData.BoneDataOffset, SkeletonBoneIndex);
+	const FDynamicWindBoneData BoneData = BoneDataBuffer.Load<FDynamicWindBoneData>(BoneDataAddress);
+
+#if USE_SIN_WIND
+	const float3 WindRotationVector = cross(float3(0.0f, 0.0f, 1.0f), WindDirection);
+	FBoneTransform BoneTransform = MakeBoneTransformIdentity();	
+	BoneTransform.Rotation = QuatNormalize(QuatFromAxisAngle(WindRotationVector, sin(Time * 0.25f) * 0.005f * WindSpeed));
+#else		
+	FBoneTransform BoneTransform = MakeBoneTransformIdentity();
+	float3 PositionAccumulation = float3(0.0f, 0.0f, 0.0f);
+	float3 LastBoneBindPosition = BoneData.BindPosePosition;
+
+	float SliceAngle = 2.0f * PI / DYNAMIC_WIND_DIRECTIONALITY_SLICES;
+
+	for (uint CurrentBoneDataAddress = BoneDataAddress, ParentBoneIndex = SkeletonBoneIndex; ParentBoneIndex != DYNAMIC_WIND_INVALID_BONE_INDEX;)
+	{
+		const FDynamicWindBoneData CurrentBoneData = BoneDataBuffer.Load<FDynamicWindBoneData>(CurrentBoneDataAddress);
+
+		// === ROTATION AXIS CALCULATIONS ========================================================================================================
+		const float3 SectionWindDirection = QuatRotateVector(QuatConjugate(QuatFromAxisAngle(float3(0.0f, 0.0f, 1.0f), Header.BlockAnimIndex * SliceAngle)), WindDirection);
+		const float3 AdjustedBoneForward = QuatRotateVector(CurrentBoneData.BindPoseRotation, float3(1.0f, 0.0f, 0.0f));
+			
+		// Prevent error when the bone forward and wind direction are collinear
+		const float3 WindHorRotVector = abs(dot(AdjustedBoneForward, SectionWindDirection)) > 0.999f ?
+			normalize(cross(AdjustedBoneForward, float3(0.0f, 0.0f, 1.0f))) :
+			normalize(cross(AdjustedBoneForward, SectionWindDirection));
+			
+		const float3 WindVerRotVector = normalize(cross(AdjustedBoneForward, float3(0.0f, 0.0f, 1.0f)));
+			
+		const float WindDirectionAlignment = 1.0f - saturate(dot(AdjustedBoneForward, SectionWindDirection)) * DYNAMIC_WIND_ALIGN_ATTENUATION;
+		const float WindDirectionOrthogonality = 1.0f - abs(dot(AdjustedBoneForward, SectionWindDirection));
+			
+		const float HeightGrad = saturate(CurrentBoneData.BindPosePosition.z / SkeletonData.SkeletonHeight);
+		// RandBBSfloat(0) returns 0, which result in an uniform rotation for all branches
+		// SkeletonDataIndex is 0 for the first registered skeleton, hence the +1
+		const float RandomValue = RandBBSfloat(((Header.SkeletonDataIndex << 8) | Header.BlockAnimIndex) + 1);
+			
+		float GustSpeed = lerp(DYNAMIC_WIND_GUST_SPEED_MIN, DYNAMIC_WIND_GUST_SPEED_MAX, WindSpeed / 100.0f);
+		float GustStrength = RampStrength(WindSpeed / 100.0f);
+
+		FQuat PartialRotation;
+
+		if (CurrentBoneData.bIsTrunkBone)
+		{
+			float WindX = WindWave(GustSpeed * Time + RandomValue * 17) + 0.75f;
+			WindX *= lerp(0.3f, 0.6f, HeightGrad);
+			WindX *= 0.06f * GustStrength;
+			WindX += 0.02f * Pow4(WindSpeed / 100.0f);
+				
+			float WindY = WindWave(GustSpeed * Time + RandomValue * 17 + 23);
+			WindY *= lerp(0.5f, 0.4f, HeightGrad);
+			WindY *= 0.03f * GustStrength;
+			WindY *= 0.25f;
+				
+			float3 BendX = cross(float3(0.0f, 0.0f, 1.0f), SectionWindDirection);
+			float3 BendY = cross(float3(0.0f, 0.0f, 1.0f), BendX);
+				
+			FQuat RotateX = QuatFromAxisAngle(BendX, WindX);
+			FQuat RotateY = QuatFromAxisAngle(BendY, WindY);
+				
+			PartialRotation = QuatMultiply(RotateX, RotateY);
+			PartialRotation = QuatAttenuate(PartialRotation, (1.0f - SkeletonData.GustAttenuation) * WindAmplitude);
+		}
+		else // Branches
+		{
+			const float DistanceFromOrigin = distance(CurrentBoneData.BindPosePosition, float3(0.0f, 0.0f, 0.0f));
+			const float BlendByGroup = max(CurrentBoneData.SimulationGroupIndex, 0) / (SkeletonData.MaxSimulationGroupIndex + 1);
+			const float TrunkBlend = (float)CurrentBoneData.bIsTrunkBone;
+				
+			float Gust = GustStrength * max(0, WindWave(GustSpeed * (Time + 0.5f) + RandomValue * 17));
+
+			float WindFlowIntensity = WindFlow(RandomValue, DistanceFromOrigin, BlendByGroup, CurrentBoneData.BoneChainLength, TrunkBlend);
+			WindFlowIntensity *= lerp(0.2f, 0.4f, Pow2(HeightGrad));
+			WindFlowIntensity *= Gust * (4 + 8 * cos(0.4f * Time + RandomValue * 17 + frac(DistanceFromOrigin * RandomValue) * 1777)) + 1.0f;
+
+			// Soft clip
+			WindFlowIntensity /= PI;
+			WindFlowIntensity /= WindFlowIntensity + 1;
+			WindFlowIntensity *= PI;
+
+			PartialRotation = QuatFromAxisAngle(WindHorRotVector, WindFlowIntensity * CurrentBoneData.Influence * WindDirectionAlignment);
+			
+			// Add torsion to branches positioned orthogonally to wind direction
+			const FQuat TorsionRotation = QuatFromAxisAngle(AdjustedBoneForward, DYNAMIC_WIND_TORSION_MOD * WindFlowIntensity * CurrentBoneData.Influence * WindDirectionOrthogonality);
+			PartialRotation = QuatMultiply(PartialRotation, TorsionRotation);
+				
+			float WindBounce = WindWave(0.2f * Time + RandomValue * 13 + frac(DistanceFromOrigin * RandomValue) * 1777);
+			WindBounce *= 0.25f * RampStrength(WindSpeed / 100.0f);
+			WindBounce *= Gust * 8 + 1;
+				
+			FQuat VerticalRotation = QuatFromAxisAngle(WindVerRotVector, WindBounce * CurrentBoneData.Influence);
+			PartialRotation = QuatMultiply(VerticalRotation, PartialRotation);
+
+			// Attenuate branch final rotation on lower branches
+			// It's here and not on the branch section because I want to influence all the rotation on the branch
+			const float RawHeightBlend = CurrentBoneData.BindPosePosition.z / SkeletonData.SkeletonHeight;
+			const float FinalHeightBlend = max(SkeletonData.bIsGroundCover, RemapValue(RawHeightBlend, 0.0f, 1.0f, DYNAMIC_WIND_HEIGHT_ATTENUATION, 1.0f));
+			PartialRotation = QuatAttenuate(PartialRotation, FinalHeightBlend * WindAmplitude);
+		}
+			
+		const FQuat CurrentRotation = QuatMultiply(BoneTransform.Rotation, PartialRotation);
+			
+		PositionAccumulation = QuatRotateVector(PartialRotation, PositionAccumulation + LastBoneBindPosition - CurrentBoneData.BindPosePosition);
+		BoneTransform.Rotation = CurrentRotation;
+		LastBoneBindPosition = CurrentBoneData.BindPosePosition;
+			
+		ParentBoneIndex = CurrentBoneData.ParentBoneIndex;
+		CurrentBoneDataAddress = GetBoneDataAddress(SkeletonData.BoneDataOffset, CurrentBoneData.ParentBoneIndex);
+	}
+		
+	BoneTransform.Location = PositionAccumulation - QuatRotateVector(BoneTransform.Rotation, BoneData.BindPosePosition);
+#endif
+
+	const float4x3 CurrentTransform = (float4x3)MatrixFromBoneTransform(BoneTransform);
+	const float4x3 PreviousTransform = LoadCompressedBoneTransform(TransformBuffer, Header.BlockDstTransformOffset, BoneIndexInBlock);
+	
+	// TODO: Handle first frame by writing CurrentTransform for both
+
+	StoreCompressedBoneTransform(TransformBuffer, Header.BlockDstTransformOffset, BoneIndexInBlock, CurrentTransform);
+	StoreCompressedBoneTransform(TransformBuffer, Header.BlockDstTransformOffset, BoneIndexInBlock + Header.TotalTransformCount, PreviousTransform);
+}
\ No newline at end of file
diff --git a/Engine/Plugins/Experimental/DynamicWind/Shaders/Shared/DynamicWindCommon.ush b/Engine/Plugins/Experimental/DynamicWind/Shaders/Shared/DynamicWindCommon.ush
new file mode 100644
--- /dev/null
+++ b/Engine/Plugins/Experimental/DynamicWind/Shaders/Shared/DynamicWindCommon.ush
@@ -0,0 +1,67 @@
+// Copyright Epic Games, Inc. All Rights Reserved.
+
+#pragma once
+
+#define DYNAMIC_WIND_INVALID_BONE_INDEX 65535
+
+#ifdef __cplusplus
+namespace UE::HLSL
+{
+#endif //__cplusplus
+
+struct FDynamicWindSkeletonData
+{
+	/** The index of the first element of the FDynamicWindBoneData for this skeleton */
+	uint BoneDataOffset;
+	/** The height (in units) of the skeleton */
+	float SkeletonHeight;
+	/** How much to attenuate trunk bones in the skeleton (0-1) */
+	float GustAttenuation;
+	/** The maximum SimulationGroupIndex value for all bones in the skeleton */
+	uint MaxSimulationGroupIndex : 31;
+	/** Whether or not this object is ground cover */
+	uint bIsGroundCover : 1;
+};
+
+struct FDynamicWindBoneData
+{
+	/** The index of this bone's parent bone, in the skeleton */
+	uint ParentBoneIndex;
+	/** The wind influence of this bone */
+	float Influence;
+	/** The simulation group for this bone */
+	int SimulationGroupIndex;
+	/** The length of bone chain that this bone is a member of */
+	float BoneChainLength;
+	/** Initial local-space orientation of the bone */
+	float4 BindPoseRotation;
+	/** Initial local-space position of the bone */
+	float3 BindPosePosition;
+	/** Whether this bone is to be considered part of a tree trunk */
+	uint bIsTrunkBone;
+};
+
+struct FDynamicWindBlockHeader
+{	
+	/** Index of the related FDynamicWindSkeletonData for the block */
+	uint SkeletonDataIndex;
+	/** Byte offset of the transforms to be written in the transform buffer */
+	uint BlockDstTransformOffset;
+	/** Number of bones in this block */
+	uint BlockNumBones : 8;
+	/** Index of the first bone (in the skeleton) of the block */
+	uint BlockBoneOffset : 24;
+	/** Index of the unique anim this block is associated with */
+	uint BlockAnimIndex : 8;
+	/** Total number of transforms to be written for this skeleton anim. */
+	uint TotalTransformCount : 24;
+};
+
+#ifdef __cplusplus
+} // namespace UE::HLSL
+
+using FDynamicWindSkeletonData = UE::HLSL::FDynamicWindSkeletonData;
+using FDynamicWindBoneData = UE::HLSL::FDynamicWindBoneData;
+using FDynamicWindBlockHeader = UE::HLSL::FDynamicWindBlockHeader;
+
+#endif //__cplusplus
diff --git a/Engine/Plugins/Experimental/DynamicWind/Source/DynamicWind/DynamicWind.Build.cs b/Engine/Plugins/Experimental/DynamicWind/Source/DynamicWind/DynamicWind.Build.cs
new file mode 100644
--- /dev/null
+++ b/Engine/Plugins/Experimental/DynamicWind/Source/DynamicWind/DynamicWind.Build.cs
@@ -0,0 +1,71 @@
+// Copyright Epic Games, Inc. All Rights Reserved.
+
+using System.IO;
+using UnrealBuildTool;
+
+namespace UnrealBuildTool.Rules
+{
+	public class DynamicWind : ModuleRules
+	{
+		public DynamicWind(ReadOnlyTargetRules Target) : base(Target)
+		{
+			PCHUsage = ModuleRules.PCHUsageMode.UseExplicitOrSharedPCHs;
+
+			PrivateIncludePaths.AddRange(
+				new string[] {
+					Path.Combine(GetModuleDirectory("Renderer"), "Private"),
+					Path.Combine(PluginDirectory, "Shaders")
+				}
+			);
+
+			PublicDependencyModuleNames.AddRange(
+				new string[]
+				{
+					"Core",
+					"CoreUObject",
+					"Engine",
+				}
+			);
+
+			PrivateDependencyModuleNames.AddRange(
+				new string[]
+				{
+					"Projects",
+					"CoreUObject",
+					"Engine",
+					"SlateCore",
+					"Slate",
+					"RenderCore",
+					"Renderer",
+					"RHICore",
+					"RHI",
+					"RHI",
+					"SlateCore",
+					"Foliage"
+				}
+			);
+
+			if (Target.bBuildEditor)
+			{
+				PrivateDependencyModuleNames.AddRange(
+					new string[]
+					{
+						"DerivedDataCache",
+						"ImageCore",
+						"MeshBuilder",
+						"MeshBuilderCommon",
+						"MeshDescription",
+						"MeshUtilities",
+						"MeshUtilitiesCommon",
+						"NaniteUtilities",
+						"RawMesh",
+						"StaticMeshDescription",
+					}
+				);
+
+				PrivateIncludePathModuleNames.Add("NaniteBuilder");
+				DynamicallyLoadedModuleNames.Add("NaniteBuilder");
+			}
+		}
+	}
+}
diff --git a/Engine/Plugins/Experimental/DynamicWind/Source/DynamicWind/Internal/DynamicWindLog.h b/Engine/Plugins/Experimental/DynamicWind/Source/DynamicWind/Internal/DynamicWindLog.h
new file mode 100644
--- /dev/null
+++ b/Engine/Plugins/Experimental/DynamicWind/Source/DynamicWind/Internal/DynamicWindLog.h
@@ -0,0 +1,7 @@
+// Copyright Epic Games, Inc. All Rights Reserved.
+
+#pragma once
+
+#include "Logging/LogMacros.h"
+
+DYNAMICWIND_API DECLARE_LOG_CATEGORY_EXTERN(LogDynamicWind, Log, All);
\ No newline at end of file
diff --git a/Engine/Plugins/Experimental/DynamicWind/Source/DynamicWind/Private/DynamicWindData.cpp b/Engine/Plugins/Experimental/DynamicWind/Source/DynamicWind/Private/DynamicWindData.cpp
new file mode 100644
--- /dev/null
+++ b/Engine/Plugins/Experimental/DynamicWind/Source/DynamicWind/Private/DynamicWindData.cpp
@@ -0,0 +1,87 @@
+// Copyright Epic Games, Inc. All Rights Reserved.
+
+#include "DynamicWindData.h"
+#include "DynamicWindSkeletalData.h"
+#include "DynamicWindProvider.h"
+#include "RenderingThread.h"
+#include "RenderUtils.h"
+#include "ScenePrivate.h"
+#include "SkeletalRenderPublic.h"
+#include "SkinningDefinitions.h"
+
+bool UDynamicWindData::IsEnabled() const
+{
+	return bEnabled;
+}
+
+const FGuid& UDynamicWindData::GetTransformProviderID() const
+{
+	static FGuid DynamicWindProviderId(DYNAMIC_WIND_TRANSFORM_PROVIDER_GUID);
+	return DynamicWindProviderId;
+}
+
+bool UDynamicWindData::UsesSkeletonBatching() const
+{
+	return true;
+}
+
+const uint32 UDynamicWindData::GetUniqueAnimationCount() const
+{
+	return 1;
+}
+
+bool UDynamicWindData::HasAnimationBounds() const
+{
+	return false;
+}
+
+bool UDynamicWindData::GetAnimationBounds(uint32 AnimationIndex, FRenderBounds& OutBounds) const
+{
+	return false;
+}
+
+FTransformProviderRenderProxy* UDynamicWindData::CreateRenderThreadResources(FSkeletalMeshSceneProxyBase* SceneProxy, FSceneInterface& Scene, FRHICommandListBase& RHICmdList)
+{
+	FDynamicWindDataRenderProxy* ProviderProxy = new FDynamicWindDataRenderProxy(this, SceneProxy, Scene);
+	ProviderProxy->CreateRenderThreadResources(RHICmdList);
+	return ProviderProxy;
+}
+
+void UDynamicWindData::DestroyRenderThreadResources(FTransformProviderRenderProxy* ProviderProxy)
+{
+	ProviderProxy->DestroyRenderThreadResources();
+	delete ProviderProxy;
+}
+
+FDynamicWindDataRenderProxy::FDynamicWindDataRenderProxy(UDynamicWindData* BankData, FSkeletalMeshSceneProxyBase* InSceneProxy, FSceneInterface& InScene)
+: SceneProxy(InSceneProxy)
+, Scene(InScene)
+{
+	UniqueAnimationCount = SceneProxy->GetUniqueAnimationCount();
+	const USkinnedAsset* SkinnedAsset = SceneProxy->GetSkinnedAsset();
+	check(SkinnedAsset != nullptr);
+	
+	const USkeletalMesh* SkeletalMesh = CastChecked<USkeletalMesh>(SkinnedAsset);
+	SkeletalData = const_cast<USkeletalMesh*>(SkeletalMesh)->GetAssetUserData<UDynamicWindSkeletalData>();
+	check(SkeletalData);
+
+	// GW-TODO: Need a validation method on the base data type to check per-proxy if ref pose fallback is needed.
+}
+
+FDynamicWindDataRenderProxy::~FDynamicWindDataRenderProxy()
+{
+}
+
+void FDynamicWindDataRenderProxy::CreateRenderThreadResources(FRHICommandListBase& RHICmdList)
+{
+}
+
+void FDynamicWindDataRenderProxy::DestroyRenderThreadResources()
+{
+}
+
+const TConstArrayView<uint64> FDynamicWindDataRenderProxy::GetProviderData(bool& bOutValid) const
+{
+	bOutValid = true;
+	return {};
+}
\ No newline at end of file
diff --git a/Engine/Plugins/Experimental/DynamicWind/Source/DynamicWind/Private/DynamicWindLog.cpp b/Engine/Plugins/Experimental/DynamicWind/Source/DynamicWind/Private/DynamicWindLog.cpp
new file mode 100644
--- /dev/null
+++ b/Engine/Plugins/Experimental/DynamicWind/Source/DynamicWind/Private/DynamicWindLog.cpp
@@ -0,0 +1,5 @@
+// Copyright Epic Games, Inc. All Rights Reserved.
+
+#include "DynamicWindLog.h"
+
+DEFINE_LOG_CATEGORY(LogDynamicWind);
diff --git a/Engine/Plugins/Experimental/DynamicWind/Source/DynamicWind/Private/DynamicWindModule.cpp b/Engine/Plugins/Experimental/DynamicWind/Source/DynamicWind/Private/DynamicWindModule.cpp
new file mode 100644
--- /dev/null
+++ b/Engine/Plugins/Experimental/DynamicWind/Source/DynamicWind/Private/DynamicWindModule.cpp
@@ -0,0 +1,40 @@
+// Copyright Epic Games, Inc. All Rights Reserved.
+
+#include "Modules/ModuleManager.h"
+#include <Interfaces/IPluginManager.h>
+#include "Misc/Paths.h"
+#include "ShaderCore.h"
+
+#define LOCTEXT_NAMESPACE "DynamicWind"
+
+class FDynamicWindModule : public IModuleInterface
+{
+public:
+	virtual void StartupModule() override
+	{
+		FCoreDelegates::OnPostEngineInit.AddRaw(this, &FDynamicWindModule::OnPostEngineInit);
+		FCoreDelegates::OnEnginePreExit.AddRaw(this, &FDynamicWindModule::OnEnginePreExit);
+
+		const FString PluginShaderDir = FPaths::Combine(IPluginManager::Get().FindPlugin(TEXT("DynamicWind"))->GetBaseDir(), TEXT("Shaders"));
+		AddShaderSourceDirectoryMapping(TEXT("/Plugin/DynamicWind"), PluginShaderDir);
+	}
+
+	virtual void ShutdownModule() override
+	{
+		FCoreDelegates::OnEnginePreExit.RemoveAll(this);
+		FCoreDelegates::OnPostEngineInit.RemoveAll(this);
+	}
+
+private:
+	void OnPostEngineInit()
+	{
+	}
+
+	void OnEnginePreExit()
+	{
+	}
+};
+
+IMPLEMENT_MODULE(FDynamicWindModule, DynamicWind);
+
+#undef LOCTEXT_NAMESPACE
diff --git a/Engine/Plugins/Experimental/DynamicWind/Source/DynamicWind/Private/DynamicWindProvider.cpp b/Engine/Plugins/Experimental/DynamicWind/Source/DynamicWind/Private/DynamicWindProvider.cpp
new file mode 100644
--- /dev/null
+++ b/Engine/Plugins/Experimental/DynamicWind/Source/DynamicWind/Private/DynamicWindProvider.cpp
@@ -0,0 +1,452 @@
+// Copyright Epic Games, Inc. All Rights Reserved.
+
+#include "DynamicWindProvider.h"
+#include "DynamicWindSkeletalData.h"
+#include "Animation/Skeleton.h"
+#include "Engine/Texture.h"
+#include "GlobalRenderResources.h"
+#include "GlobalShader.h"
+#include "RenderGraphBuilder.h"
+#include "RenderGraphResources.h"
+#include "RenderGraphUtils.h"
+#include "RenderingThread.h"
+#include "RenderUtils.h"
+#include "ScenePrivate.h"
+#include "SkeletalRenderPublic.h"
+#include "SkinningDefinitions.h"
+#include "TextureResource.h"
+
+static FGuid GDynamicWindTransformProviderId(DYNAMIC_WIND_TRANSFORM_PROVIDER_GUID);
+
+DECLARE_GPU_STAT(DynamicWind);
+
+namespace
+{
+
+TAutoConsoleVariable<float> CVarDynamicWindOverrideSpeed(
+	TEXT("DynamicWind.OverrideSpeed"),
+	-1.0f,
+	TEXT("Overrides the wind speed. If this value is below 0.0f, then use the default wind speed set on the dynamic wind actor."),
+	ECVF_RenderThreadSafe | ECVF_Default
+);
+
+TAutoConsoleVariable<float> CVarDynamicWindOverrideAmplitude(
+	TEXT("DynamicWind.OverrideAmplitude"),
+	-1.0f,
+	TEXT("Overrides the wind amplitude. If this value is below 0.0f, then use the default wind amplitude set on the dynamic wind actor."),
+	ECVF_RenderThreadSafe | ECVF_Default
+);
+
+TAutoConsoleVariable<float> CVarDynamicWindOverrideRateOfChange(
+	TEXT("DynamicWind.OverrideRateOfChange"),
+	0.5f,
+	TEXT("Overrides the rate of change per second of the wind amplitude. It is used to smooth the wind amplitude over time to prevent animation popping."),
+	ECVF_RenderThreadSafe | ECVF_Default
+);
+
+TAutoConsoleVariable<bool> CVarDynamicWindUseSine(
+	TEXT("DynamicWind.UseSine"),
+	false,
+	TEXT("Forces dynamic wind to use a basic sine function to transform the bones."),
+	ECVF_RenderThreadSafe | ECVF_Default
+);
+
+TAutoConsoleVariable<bool> CVarDynamicWindPreferAsyncCompute(
+	TEXT("DynamicWind.PreferAsyncCompute"),
+	false,
+	TEXT("Run dynamic wind on async compute, if supported."),
+	ECVF_RenderThreadSafe | ECVF_Default
+);
+
+constexpr uint32 INVALID_BONE_OFFSET = ~uint32(0);
+
+bool IsValidBoneOffset(const uint32 BoneOffset)
+{
+	return BoneOffset != INVALID_BONE_OFFSET;
+}
+
+bool IsDynamicWindOnAsyncCompute()
+{
+	return GSupportsEfficientAsyncCompute && CVarDynamicWindPreferAsyncCompute.GetValueOnAnyThread();
+}
+
+}
+
+class FDynamicWindEvalCS : public FGlobalShader
+{
+public:
+	static constexpr uint32 BonesPerGroup = 64u;
+
+private:
+	DECLARE_GLOBAL_SHADER(FDynamicWindEvalCS);
+	SHADER_USE_PARAMETER_STRUCT(FDynamicWindEvalCS, FGlobalShader);
+
+	class FUseSinFunctionDim : SHADER_PERMUTATION_BOOL("USE_SIN_WIND");
+	using FPermutationDomain = TShaderPermutationDomain<FUseSinFunctionDim>;
+
+	static bool ShouldCompilePermutation(const FGlobalShaderPermutationParameters& Parameters)
+	{
+		return DoesPlatformSupportNanite(Parameters.Platform);
+	}
+
+	static void ModifyCompilationEnvironment(const FGlobalShaderPermutationParameters& Parameters, FShaderCompilerEnvironment& OutEnvironment)
+	{
+		FGlobalShader::ModifyCompilationEnvironment(Parameters, OutEnvironment);
+
+		OutEnvironment.CompilerFlags.Add(CFLAG_WarningsAsErrors);
+		OutEnvironment.CompilerFlags.Add(CFLAG_HLSL2021);
+
+		OutEnvironment.SetDefine(TEXT("BONES_PER_GROUP"), BonesPerGroup);
+	}
+
+	BEGIN_SHADER_PARAMETER_STRUCT(FParameters, )
+		SHADER_PARAMETER_RDG_BUFFER_SRV(StructuredBuffer<FDynamicWindSkeletonData>, SkeletonBuffer)
+		SHADER_PARAMETER_RDG_BUFFER_SRV(StructuredBuffer<FDynamicWindBlockHeader>, HeaderBuffer)
+		SHADER_PARAMETER_RDG_BUFFER_SRV(ByteAddressBuffer, BoneDataBuffer)
+		SHADER_PARAMETER_RDG_BUFFER_UAV(RWByteAddressBuffer, TransformBuffer)
+		SHADER_PARAMETER_RDG_TEXTURE_SRV(Texture2D<float4>, WindTexture)
+		SHADER_PARAMETER_SAMPLER(SamplerState, WindSampler)
+		SHADER_PARAMETER(float, Time)
+		SHADER_PARAMETER(float, DeltaTime)
+		SHADER_PARAMETER(float, WindSpeed)
+		SHADER_PARAMETER(float, WindAmplitude)
+		SHADER_PARAMETER(FVector3f, WindDirection)
+		SHADER_PARAMETER(FVector4f, Debug)
+	END_SHADER_PARAMETER_STRUCT()
+};
+
+IMPLEMENT_GLOBAL_SHADER(FDynamicWindEvalCS, "/Plugin/DynamicWind/DynamicWindEval.usf", "WindEvalCS", SF_Compute);
+
+FDynamicWindTransformProvider::FDynamicWindTransformProvider(FScene& InScene)
+: Scene(InScene)
+, BoneDataBuffer(1, TEXT("DynamicWind.BoneDataBuffer"))
+{
+	if (auto TransformProvider = Scene.GetExtensionPtr<FSkinningTransformProvider>())
+	{
+		// Register GPU dynamic wind transform provider
+		TransformProvider->RegisterProvider(
+			GDynamicWindTransformProviderId,
+			FSkinningTransformProvider::FOnProvideTransforms::CreateRaw(this, &FDynamicWindTransformProvider::ProvideTransforms),
+			true /* Use skeleton batching */
+		);
+	}
+}
+
+FDynamicWindTransformProvider::~FDynamicWindTransformProvider()
+{
+	if (auto TransformProvider = Scene.GetExtensionPtr<FSkinningTransformProvider>())
+	{
+		TransformProvider->UnregisterProvider(GDynamicWindTransformProviderId);
+	}
+}
+	
+void FDynamicWindTransformProvider::RegisterSceneProxy(const Nanite::FSkinnedSceneProxy* InProxy)
+{
+	USkeletalMesh* SkeletalMesh = const_cast<USkeletalMesh*>(Cast<USkeletalMesh>(InProxy->GetSkinnedAsset()));
+	const USkeleton* Skeleton = SkeletalMesh->GetSkeleton();
+	const FReferenceSkeleton& RefSkeleton = Skeleton->GetReferenceSkeleton();
+	const uint32 NumBonesInSkeleton = RefSkeleton.GetNum();
+	const UDynamicWindSkeletalData* Data = SkeletalMesh->GetAssetUserData<UDynamicWindSkeletalData>();
+	const uint64 UserDataHash = Data->GetSkeletalDataHash();
+	
+	FSkeletonEntry& SkeletonEntry = SkeletonLookup.FindOrAdd(Skeleton->GetGuid());
+	
+	++SkeletonEntry.ReferenceCount;
+	if (SkeletonEntry.ReferenceCount == 1)
+	{
+		SkeletonEntry.Data.BoneDataOffset = BoneDataAllocator.Allocate(NumBonesInSkeleton);
+		SkeletonEntry.NumBones = NumBonesInSkeleton;
+	#if DYNAMIC_WIND_DEBUG_SKELETON_NAMES
+		SkeletonEntry.SkeletonName = Skeleton->GetFName();
+	#endif
+	}
+	else
+	{
+		check(SkeletonEntry.Data.BoneDataOffset != INDEX_NONE);
+		check(SkeletonEntry.NumBones == NumBonesInSkeleton);
+		if (SkeletonEntry.UserDataHash == UserDataHash)
+		{
+			// no upload necessary
+			return;
+		}
+	}
+
+	if (!BoneDataUploader.IsValid())
+	{
+		BoneDataUploader = MakeUnique<FBoneDataUploader>();
+	}
+
+	const FBoxSphereBounds Bounds = SkeletalMesh->GetBounds();
+
+	SkeletonEntry.UserDataHash = UserDataHash;
+	SkeletonEntry.Data.MaxSimulationGroupIndex = uint32(FMath::Max(Data->SimulationGroups.Num() - 1, 0));
+	SkeletonEntry.Data.SkeletonHeight = Bounds.BoxExtent.Z + Bounds.Origin.Z;
+	SkeletonEntry.Data.GustAttenuation = Data->GustAttenuation;
+	SkeletonEntry.Data.bIsGroundCover = Data->bIsGroundCover;
+
+	TArrayView<FDynamicWindBoneData> BoneData = BoneDataUploader->AddMultiple_GetRef(SkeletonEntry.Data.BoneDataOffset, NumBonesInSkeleton);
+	for (uint32 BoneIndex = 0; BoneIndex < NumBonesInSkeleton; ++BoneIndex)
+	{
+		const uint32 ParentBoneIndex = FMath::Min<uint32>(RefSkeleton.GetParentIndex(BoneIndex), DYNAMIC_WIND_INVALID_BONE_INDEX);			
+		const FTransform AbsoluteTransform = RefSkeleton.GetBoneAbsoluteTransform(BoneIndex);
+		const FQuat AbsoluteRotation = AbsoluteTransform.GetRotation();
+		const FVector3f BindPosePosition = FVector3f(AbsoluteTransform.GetLocation());
+		const FVector4f BindPoseRotation = FVector4f(
+			AbsoluteRotation.X,
+			AbsoluteRotation.Y,
+			AbsoluteRotation.Z,
+			AbsoluteRotation.W
+		);
+
+		BoneData[BoneIndex] = FDynamicWindBoneData
+		{
+			.ParentBoneIndex = ParentBoneIndex,
+			.Influence = Data->GetBoneInfluence(BoneIndex),
+			.SimulationGroupIndex = Data->GetBoneSimulationGroupIndex(BoneIndex),
+			.BoneChainLength = Data->GetBoneChainLength(BoneIndex),
+			.BindPoseRotation = BindPoseRotation,
+			.BindPosePosition = BindPosePosition,
+			.bIsTrunkBone = Data->IsTrunkBone(BoneIndex) ? 1u : 0u,
+		};
+	}
+}
+
+void FDynamicWindTransformProvider::UnregisterSceneProxy(const Nanite::FSkinnedSceneProxy* InProxy)
+{
+	const USkeletalMesh* SkeletalMesh = Cast<USkeletalMesh>(InProxy->GetSkinnedAsset());
+	const USkeleton* Skeleton = SkeletalMesh->GetSkeleton();
+	auto EntryID = SkeletonLookup.FindId(Skeleton->GetGuid());
+
+	if (!SkeletonLookup.IsValidId(EntryID))
+	{
+		return;
+	}
+
+	FSkeletonEntry& SkeletonEntry = SkeletonLookup.Get(EntryID).Value;
+	--SkeletonEntry.ReferenceCount;
+	if (SkeletonEntry.ReferenceCount == 0)
+	{
+		// Last proxy out, remove the whole entry
+		BoneDataAllocator.Free(SkeletonEntry.Data.BoneDataOffset, SkeletonEntry.NumBones);
+		SkeletonLookup.Remove(EntryID);
+	}
+}
+
+int32 FDynamicWindTransformProvider::UpdateParameters(const FDynamicWindParameters& Parameters)
+{
+	WindParameters = Parameters;
+	return INDEX_NONE;
+}
+
+float FDynamicWindTransformProvider::GetBlendedWindAmplitude() const
+{
+	return BlendedWindAmplitude;
+}
+
+void FDynamicWindTransformProvider::ProvideTransforms(FSkinningTransformProvider::FProviderContext& Context)
+{
+	const bool bHasAlreadySimulatedThisFrame = LastSimulatedFrameNumber == Scene.GetFrameNumber();
+	
+	// There are multiple broadcasts from SkinningSceneExtension, for primary view as well virtual texture updates.
+	// But we need to run the simulation only once.
+	if (bHasAlreadySimulatedThisFrame || SkeletonLookup.IsEmpty())
+	{
+		return;
+	}
+	
+	const uint32 BonesPerGroup = FDynamicWindEvalCS::BonesPerGroup;
+	uint32 BlockCount = 0;
+	for (const FUintVector2& Indirection : Context.Indirections)
+	{
+		const FSkeletonBatch& SkeletonBatch = Context.SkeletonBatches[Indirection.X];
+		const uint32 TransformCount = SkeletonBatch.MaxBoneTransforms;
+		BlockCount += FMath::DivideAndRoundUp(TransformCount, BonesPerGroup) * SkeletonBatch.UniqueAnimationCount;
+	}
+
+	if (BlockCount == 0)
+	{
+		// Nothing to simulate
+		return;
+	}
+
+	FRDGBuilder& GraphBuilder = Context.GraphBuilder;
+
+	// Create the skeleton data and block header buffers
+	const int32 SkeletonCount = SkeletonLookup.Num();
+	UE::HLSL::FDynamicWindSkeletonData* SkeletonData = GraphBuilder.AllocPODArray<UE::HLSL::FDynamicWindSkeletonData>(SkeletonCount);
+	UE::HLSL::FDynamicWindBlockHeader* BlockHeaders = GraphBuilder.AllocPODArray<UE::HLSL::FDynamicWindBlockHeader>(BlockCount);
+
+	GraphBuilder.AddSetupTask(
+		[
+			BlockHeaders,
+			SkeletonData,
+			SkeletonLookup = TMap<FGuid, FSkeletonEntry, SceneRenderingSetAllocator>(SkeletonLookup), // Copy for task
+			Indirections = Context.Indirections,
+			SkeletonBatches = Context.SkeletonBatches
+		]
+		{
+			TRACE_CPUPROFILER_EVENT_SCOPE(FDynamicWindTransformProvider::ProvideTransforms);
+
+			int32 SkeletonIndex = 0;
+			for (const auto& [_, SkeletonEntry] : SkeletonLookup)
+			{
+				SkeletonData[SkeletonIndex] = SkeletonEntry.Data;
+				++SkeletonIndex;
+			}
+
+			uint32 BlockWrite = 0;
+			for (const FUintVector2& Indirection : Indirections)
+			{
+				const FSkeletonBatch& SkinnedAnimationBatch = SkeletonBatches[Indirection.X];
+
+				const uint32 TransformCount = SkinnedAnimationBatch.MaxBoneTransforms;
+				const uint32 UniqueAnimationCount = SkinnedAnimationBatch.UniqueAnimationCount;
+
+				const uint32 FullBlockCount = TransformCount / BonesPerGroup;
+				const uint32 PartialTransformCount = TransformCount - (FullBlockCount * BonesPerGroup);
+
+				uint32 DstTransformOffset = Indirection.Y;
+
+				FSetElementId SkeletonId = SkeletonLookup.FindId(SkinnedAnimationBatch.SkeletonGuid);
+
+				for (uint32 AnimIndex = 0; AnimIndex < UniqueAnimationCount; ++AnimIndex)
+				{
+					uint32 BoneOffset = 0;
+
+					for (uint32 BlockIndex = 0; BlockIndex < FullBlockCount; ++BlockIndex)
+					{
+						BlockHeaders[BlockWrite].SkeletonDataIndex = SkeletonId.AsInteger(); // NOTE: Can be INDEX_NONE
+						BlockHeaders[BlockWrite].BlockDstTransformOffset = DstTransformOffset;
+						BlockHeaders[BlockWrite].BlockNumBones = BonesPerGroup;
+						BlockHeaders[BlockWrite].BlockBoneOffset = BoneOffset;
+						BlockHeaders[BlockWrite].BlockAnimIndex = AnimIndex;
+						BlockHeaders[BlockWrite].TotalTransformCount = TransformCount;
+						
+						BoneOffset += BonesPerGroup;
+						DstTransformOffset += (BonesPerGroup * sizeof(FCompressedBoneTransform));
+						++BlockWrite;
+					}
+
+					if (PartialTransformCount > 0)
+					{
+						BlockHeaders[BlockWrite].SkeletonDataIndex = SkeletonId.AsInteger(); // NOTE: Can be INDEX_NONE
+						BlockHeaders[BlockWrite].BlockDstTransformOffset = DstTransformOffset;
+						BlockHeaders[BlockWrite].BlockNumBones = PartialTransformCount;
+						BlockHeaders[BlockWrite].BlockBoneOffset = BoneOffset;
+						BlockHeaders[BlockWrite].BlockAnimIndex = AnimIndex;
+						BlockHeaders[BlockWrite].TotalTransformCount = TransformCount;
+
+						BoneOffset += PartialTransformCount;
+						DstTransformOffset += (PartialTransformCount * sizeof(FCompressedBoneTransform));
+						++BlockWrite;
+					}
+
+					// Advance past all the previous transforms as well
+					DstTransformOffset += (sizeof(FCompressedBoneTransform) * TransformCount);
+				}
+			}
+		},
+		UE::Tasks::ETaskPriority::High
+	);
+
+	FRDGBufferRef SkeletonBuffer = CreateStructuredBuffer(
+		GraphBuilder,
+		TEXT("DynamicWind.SkeletonData"),
+		MakeConstArrayView(SkeletonData, SkeletonCount),
+		ERDGInitialDataFlags::NoCopy // The buffer data is allocated above on the RDG timeline
+	);
+
+	FRDGBufferRef BlockHeaderBuffer = CreateStructuredBuffer(
+		GraphBuilder,
+		TEXT("DynamicWind.ScatterHeaders"),
+		MakeConstArrayView(BlockHeaders, BlockCount),
+		ERDGInitialDataFlags::NoCopy // The buffer data is allocated above on the RDG timeline
+	);
+
+	FRDGBufferRef BoneDataBufferRDG = nullptr; 
+	if (BoneDataUploader.IsValid())
+	{
+		BoneDataBufferRDG = BoneDataUploader->ResizeAndUploadTo(
+			GraphBuilder,
+			BoneDataBuffer,
+			BoneDataAllocator.GetMaxSize()
+		);
+		BoneDataUploader = nullptr;
+	}
+	else
+	{
+		BoneDataBufferRDG = BoneDataBuffer.Register(GraphBuilder);
+	}
+
+	FTextureRHIRef WindTextureRHI = nullptr;
+	FSamplerStateRHIRef WindSampler = nullptr;
+	if (WindParameters.WindTexture)
+	{
+		if (FTextureResource* WindTextureResource = WindParameters.WindTexture->GetResource())
+		{
+			WindTextureRHI = WindTextureResource->GetTextureRHI();
+			WindSampler = WindTextureResource->SamplerStateRHI;
+		}
+	}
+	if (WindTextureRHI == nullptr)
+	{
+		WindTextureRHI = GWhiteTexture->TextureRHI.GetReference();
+	}
+
+	if (WindSampler == nullptr)
+	{
+		WindSampler = TStaticSamplerState<SF_Bilinear, AM_Clamp, AM_Clamp, AM_Clamp>::GetRHI();
+	}
+
+	auto WindTextureRef = GraphBuilder.RegisterExternalTexture(CreateRenderTarget(WindTextureRHI, TEXT("DynamicWind.WindTexture")));
+
+	const float WindSpeedOverride = CVarDynamicWindOverrideSpeed.GetValueOnAnyThread();
+	const float WindSpeed = WindSpeedOverride < 0.0f ? WindParameters.WindSpeed : WindSpeedOverride;
+
+	const float WindAmplitudeOverride = CVarDynamicWindOverrideAmplitude.GetValueOnAnyThread();
+	if (!FMath::IsNearlyEqual(WindAmplitudeOverride, BlendedWindAmplitude, 1.e-2f))
+	{
+		const float RateOfChange = CVarDynamicWindOverrideRateOfChange.GetValueOnAnyThread();
+		BlendedWindAmplitude += FMath::Sign(WindAmplitudeOverride - BlendedWindAmplitude) * RateOfChange * Scene.GetWorld()->GetDeltaSeconds();
+	}
+	else
+	{
+		BlendedWindAmplitude = WindAmplitudeOverride;
+	}
+	
+	const float WindAmplitude = BlendedWindAmplitude < 0.0f ? WindParameters.WindAmplitude : BlendedWindAmplitude;
+
+	{
+		RDG_EVENT_SCOPE_STAT(GraphBuilder, DynamicWind, "DynamicWind");
+		RDG_GPU_STAT_SCOPE(GraphBuilder, DynamicWind);
+
+		FDynamicWindEvalCS::FParameters* PassParameters = GraphBuilder.AllocParameters<FDynamicWindEvalCS::FParameters>();
+		PassParameters->SkeletonBuffer = GraphBuilder.CreateSRV(SkeletonBuffer);
+		PassParameters->HeaderBuffer = GraphBuilder.CreateSRV(BlockHeaderBuffer);
+		PassParameters->BoneDataBuffer = GraphBuilder.CreateSRV(BoneDataBufferRDG);
+		PassParameters->TransformBuffer = GraphBuilder.CreateUAV(Context.TransformBuffer);
+		PassParameters->WindTexture = GraphBuilder.CreateSRV(WindTextureRef);
+		PassParameters->WindSampler = WindSampler;
+		PassParameters->Time = Scene.GetWorld()->GetTimeSeconds();
+		PassParameters->DeltaTime = Scene.GetWorld()->GetDeltaSeconds();
+		PassParameters->WindSpeed = WindSpeed;
+		PassParameters->WindAmplitude = WindAmplitude;
+		PassParameters->WindDirection = FVector3f(WindParameters.WindDirection);
+		PassParameters->Debug = WindParameters.DebugModulation;
+
+		FDynamicWindEvalCS::FPermutationDomain PermutationVectorCS;
+		PermutationVectorCS.Set<FDynamicWindEvalCS::FUseSinFunctionDim>(CVarDynamicWindUseSine.GetValueOnAnyThread());
+
+		auto ComputeShader = GetGlobalShaderMap(GMaxRHIFeatureLevel)->GetShader<FDynamicWindEvalCS>(PermutationVectorCS);
+		FComputeShaderUtils::AddPass(
+			GraphBuilder,
+			RDG_EVENT_NAME("DynamicWind Simulate [BlockCount: %d]", BlockCount),
+			IsDynamicWindOnAsyncCompute() ? ERDGPassFlags::AsyncCompute : ERDGPassFlags::Compute,
+			ComputeShader,
+			PassParameters,
+			FIntVector(BlockCount, 1, 1)
+		);
+	}
+
+	LastSimulatedFrameNumber = Scene.GetFrameNumber();
+}
diff --git a/Engine/Plugins/Experimental/DynamicWind/Source/DynamicWind/Private/DynamicWindSkeletalData.cpp b/Engine/Plugins/Experimental/DynamicWind/Source/DynamicWind/Private/DynamicWindSkeletalData.cpp
new file mode 100644
--- /dev/null
+++ b/Engine/Plugins/Experimental/DynamicWind/Source/DynamicWind/Private/DynamicWindSkeletalData.cpp
@@ -0,0 +1,104 @@
+// Copyright Epic Games, Inc. All Rights Reserved.
+
+#include "DynamicWindSkeletalData.h"
+#include "Serialization/MemoryHasher.h"
+#include "Hash/xxhash.h"
+
+int32 UDynamicWindSkeletalData::GetBoneSimulationGroupIndex(const int32 BoneIndex) const
+{
+	for (const auto& [SimulationGroupIndex, BoneIndices] : SimulationGroupBones)
+	{
+		if (BoneIndices.Contains(BoneIndex))
+		{
+			return SimulationGroupIndex;
+		}
+	}
+
+	return INDEX_NONE;
+}
+
+float UDynamicWindSkeletalData::GetBoneInfluence(const int32 BoneIndex) const
+{
+	const FDynamicWindExtraBoneData* ExtraBoneData = ExtraBonesData.Find(BoneIndex);
+	if (ExtraBoneData == nullptr)
+	{
+		return 0.0f;
+	}
+	
+	const int32 SimulationGroupIndex = GetBoneSimulationGroupIndex(BoneIndex);
+	if (!SimulationGroups.IsValidIndex(SimulationGroupIndex))
+	{
+		return 1.0f; // The user may have deleted this SimulationGroup
+	}
+
+	const FDynamicWindSimulationGroupData& SimulationGroup = SimulationGroups[SimulationGroupIndex];
+	if (!SimulationGroup.bUseDualInfluence)
+	{
+		return SimulationGroup.Influence;
+	}
+	
+	float InfluenceBlend = 0.0f;
+	if (const FDynamicWindBoneChainData* BoneChain = BoneChains.Find(ExtraBoneData->BoneChainOriginBoneIndex))
+	{		
+		InfluenceBlend = BoneChain->NumBones > 1 ? ExtraBoneData->IndexInBoneChain / (BoneChain->NumBones - 1.0f) : 0.0f;
+
+		// Formula for a bend toward zero: y = x ^ (1 / (1 - shift))
+		// Shift is biased to prevent division by zero
+		InfluenceBlend = FMath::Pow(InfluenceBlend, 1.0f / (1.0f - FMath::Min(SimulationGroup.ShiftTop, 0.99f)));
+	}
+	
+	return FMath::Lerp(SimulationGroup.MinInfluence, SimulationGroup.MaxInfluence, InfluenceBlend);
+}
+
+bool UDynamicWindSkeletalData::IsTrunkBone(const int32 BoneIndex) const
+{
+	if (!ExtraBonesData.Contains(BoneIndex))
+	{
+		return false;
+	}
+	
+	const int32 SimulationGroupIndex = GetBoneSimulationGroupIndex(BoneIndex);
+	return SimulationGroups.IsValidIndex(SimulationGroupIndex) ? SimulationGroups[SimulationGroupIndex].bIsTrunkGroup : false;
+}
+
+float UDynamicWindSkeletalData::GetBoneChainLength(const int32 BoneIndex) const
+{
+	const FDynamicWindExtraBoneData* ExtraBoneData = ExtraBonesData.Find(BoneIndex);
+	if (ExtraBoneData == nullptr)
+	{
+		return 0.0f;
+	}
+
+	if (const FDynamicWindBoneChainData* BoneChain = BoneChains.Find(ExtraBoneData->BoneChainOriginBoneIndex))
+	{
+		return BoneChain->ChainLength;
+	}
+	
+	return 0.0f;
+}
+
+void UDynamicWindSkeletalData::RecalculateSkeletalDataHash()
+{
+	TMemoryHasher<FXxHash64Builder, FXxHash64> Ar;
+	Ar.SetIsPersistent(true);
+	
+	Serialize(Ar);
+		
+	SkeletalDataHash = Ar.Finalize().Hash;
+}
+
+void UDynamicWindSkeletalData::PostLoad()
+{
+	Super::PostLoad();
+	RecalculateSkeletalDataHash();
+}
+
+#if WITH_EDITOR
+
+void UDynamicWindSkeletalData::PostEditChangeProperty(FPropertyChangedEvent& Event)
+{
+	Super::PostEditChangeProperty(Event);
+	RecalculateSkeletalDataHash();
+}
+
+#endif // WITH_EDITOR
\ No newline at end of file
diff --git a/Engine/Plugins/Experimental/DynamicWind/Source/DynamicWind/Private/DynamicWindSubsystem.cpp b/Engine/Plugins/Experimental/DynamicWind/Source/DynamicWind/Private/DynamicWindSubsystem.cpp
new file mode 100644
--- /dev/null
+++ b/Engine/Plugins/Experimental/DynamicWind/Source/DynamicWind/Private/DynamicWindSubsystem.cpp
@@ -0,0 +1,175 @@
+// Copyright Epic Games, Inc. All Rights Reserved.
+
+#include "DynamicWindSubsystem.h"
+#include "Animation/Skeleton.h"
+#include "Engine/TextureRenderTarget2D.h"
+#include "Engine/World.h"
+#include "Engine/SkinnedAsset.h"
+#include "RenderingThread.h"
+#include "TextureResource.h"
+#include "GlobalRenderResources.h"
+#include "GlobalShader.h"
+#include "RenderGraphBuilder.h"
+#include "RenderGraphResources.h"
+#include "RenderGraphUtils.h"
+#include "SceneView.h"
+#include "SceneInterface.h"
+#include "ScenePrivate.h"
+#include "NaniteSceneProxy.h"
+#include "DynamicWindSkeletalData.h"
+
+namespace
+{
+
+TAutoConsoleVariable<bool> CVarDynamicWindEnable(
+	TEXT("DynamicWind.Enable"),
+	true,
+	TEXT("Whether or not to enable the dynamic wind subsystem"),
+	ECVF_RenderThreadSafe | ECVF_ReadOnly
+);
+
+}
+
+UDynamicWindSubsystem::UDynamicWindSubsystem()
+: bInitialized(false)
+, TransformProvider(nullptr)
+, TransformProviderId(DYNAMIC_WIND_TRANSFORM_PROVIDER_GUID)
+{
+}
+
+bool UDynamicWindSubsystem::ShouldCreateSubsystem(UObject* Outer) const
+{
+	return CVarDynamicWindEnable.GetValueOnAnyThread();
+}
+
+void UDynamicWindSubsystem::PostInitialize()
+{
+	Super::PostInitialize();
+
+	if (UWorld* World = GetWorld())
+	{
+		if (FSceneInterface* Scene = World->Scene)
+		{
+			if (FScene* RenderScene = Scene->GetRenderScene())
+			{
+				TransformProvider = MakeUnique<FDynamicWindTransformProvider>(*RenderScene);
+
+				OnCreateRenderThreadResources =
+					Nanite::FSkinnedSceneProxyDelegates::OnCreateRenderThreadResources.AddLambda(
+						[this](const Nanite::FSkinnedSceneProxy* InSkinnedSceneProxy)
+						{
+							RegisterSceneProxy(InSkinnedSceneProxy);
+						}
+					);
+
+				OnDestroyRenderThreadResources =
+					Nanite::FSkinnedSceneProxyDelegates::OnDestroyRenderThreadResources.AddLambda(
+						[this](const Nanite::FSkinnedSceneProxy* InSkinnedSceneProxy)
+						{
+							UnregisterSceneProxy(InSkinnedSceneProxy);
+						}
+					);
+
+
+			}
+		}
+	}
+
+	bInitialized =
+		TransformProvider.IsValid() &&
+		OnCreateRenderThreadResources.IsValid() &&
+		OnDestroyRenderThreadResources.IsValid();
+}
+
+void UDynamicWindSubsystem::PreDeinitialize()
+{
+	bInitialized = false;
+
+	if (OnCreateRenderThreadResources.IsValid())
+	{
+		Nanite::FSkinnedSceneProxyDelegates::OnCreateRenderThreadResources.Remove(OnCreateRenderThreadResources);
+		OnCreateRenderThreadResources.Reset();
+	}
+
+	if (OnDestroyRenderThreadResources.IsValid())
+	{
+		Nanite::FSkinnedSceneProxyDelegates::OnDestroyRenderThreadResources.Remove(OnDestroyRenderThreadResources);
+		OnDestroyRenderThreadResources.Reset();
+	}
+
+	TransformProvider.Reset();
+
+	Super::PreDeinitialize();
+}
+
+void UDynamicWindSubsystem::RegisterSceneProxy(const Nanite::FSkinnedSceneProxy* InSkinnedProxy)
+{
+	if (&InSkinnedProxy->GetScene() != GetWorld()->Scene)
+	{
+		return;
+	}
+
+	if (InSkinnedProxy->GetTransformProviderId() != TransformProviderId)
+	{
+		return;
+	}
+
+	if (const USkinnedAsset* SkinnedAsset = InSkinnedProxy->GetSkinnedAsset())
+	{
+		if (!bInitialized)
+		{
+			UE_LOG(LogDynamicWind, Verbose, TEXT("Trying to register SkinnedSceneProxy (SkinnedAsset=%s), when the subsystem is not initialized."), *SkinnedAsset->GetName());
+			return;
+		}
+
+		TransformProvider->RegisterSceneProxy(InSkinnedProxy);
+	}
+}
+
+void UDynamicWindSubsystem::UnregisterSceneProxy(const Nanite::FSkinnedSceneProxy* InSkinnedProxy)
+{
+	if (&InSkinnedProxy->GetScene() != GetWorld()->Scene)
+	{
+		return;
+	}
+
+	if (InSkinnedProxy->GetTransformProviderId() != TransformProviderId)
+	{
+		return;
+	}
+
+	if (const USkinnedAsset* SkinnedAsset = InSkinnedProxy->GetSkinnedAsset())
+	{
+		if (!bInitialized)
+		{
+			UE_LOG(LogDynamicWind, Verbose, TEXT("Trying to unregister SkinnedSceneProxy (SkinnedAsset=%s), when the subsystem is not initialized."), *SkinnedAsset->GetName());
+			return;
+		}
+
+		TransformProvider->UnregisterSceneProxy(InSkinnedProxy);
+	}
+}
+
+float UDynamicWindSubsystem::GetBlendedWindAmplitude() const
+{
+	if (TransformProvider)
+	{
+		return TransformProvider->GetBlendedWindAmplitude();
+	}
+
+	return -1.0f;
+}
+
+void UDynamicWindSubsystem::UpdateWindParameters(const FDynamicWindParameters& Parameters)
+{
+	if (bInitialized)
+	{
+		ENQUEUE_RENDER_COMMAND(UpdateFoliageWindSimulationInputs)
+		(
+			[WindProvider = TransformProvider.Get(), WindParameters = Parameters](FRHICommandListImmediate&) mutable
+			{
+				WindProvider->UpdateParameters(WindParameters);
+			}
+		);
+	}
+}
diff --git a/Engine/Plugins/Experimental/DynamicWind/Source/DynamicWind/Public/DynamicWindData.h b/Engine/Plugins/Experimental/DynamicWind/Source/DynamicWind/Public/DynamicWindData.h
new file mode 100644
--- /dev/null
+++ b/Engine/Plugins/Experimental/DynamicWind/Source/DynamicWind/Public/DynamicWindData.h
@@ -0,0 +1,41 @@
+// Copyright Epic Games, Inc. All Rights Reserved.
+
+#pragma once
+
+#include "Animation/TransformProviderData.h"
+
+#include "DynamicWindData.generated.h"
+
+UCLASS(hidecategories=Object, MinimalAPI, BlueprintType)
+class UDynamicWindData : public UTransformProviderData
+{
+	GENERATED_BODY()
+
+public:
+	virtual bool IsEnabled() const override;
+	virtual const FGuid& GetTransformProviderID() const override;
+	virtual bool UsesSkeletonBatching() const override;
+	virtual const uint32 GetUniqueAnimationCount() const override;
+	virtual bool HasAnimationBounds() const override;
+	virtual bool GetAnimationBounds(uint32 AnimationIndex, FRenderBounds& OutBounds) const override;
+	virtual FTransformProviderRenderProxy* CreateRenderThreadResources(FSkeletalMeshSceneProxyBase* SceneProxy, FSceneInterface& Scene, FRHICommandListBase& RHICmdList) override;
+	virtual void DestroyRenderThreadResources(FTransformProviderRenderProxy* ProviderProxy) override;
+};
+
+class FDynamicWindDataRenderProxy : public FTransformProviderRenderProxy
+{
+	friend class UDynamicWindData;
+
+public:
+	DYNAMICWIND_API FDynamicWindDataRenderProxy(UDynamicWindData* BankData, FSkeletalMeshSceneProxyBase* SceneProxy, FSceneInterface& Scene);
+	DYNAMICWIND_API virtual ~FDynamicWindDataRenderProxy();
+	DYNAMICWIND_API virtual void CreateRenderThreadResources(FRHICommandListBase& RHICmdList) override;
+	DYNAMICWIND_API virtual void DestroyRenderThreadResources() override;
+	DYNAMICWIND_API virtual const TConstArrayView<uint64> GetProviderData(bool& bOutValid) const override;
+
+private:
+	FSkeletalMeshSceneProxyBase* SceneProxy = nullptr;
+	TObjectPtr<const class UDynamicWindSkeletalData> SkeletalData = nullptr;
+	FSceneInterface& Scene; 
+	uint32 UniqueAnimationCount = 0;
+};
diff --git a/Engine/Plugins/Experimental/DynamicWind/Source/DynamicWind/Public/DynamicWindParameters.h b/Engine/Plugins/Experimental/DynamicWind/Source/DynamicWind/Public/DynamicWindParameters.h
new file mode 100644
--- /dev/null
+++ b/Engine/Plugins/Experimental/DynamicWind/Source/DynamicWind/Public/DynamicWindParameters.h
@@ -0,0 +1,36 @@
+// Copyright Epic Games, Inc. All Rights Reserved.
+
+#pragma once
+
+#include "CoreMinimal.h"
+#include "DynamicWindParameters.generated.h"
+
+class UTexture;
+
+USTRUCT(BlueprintType)
+struct FDynamicWindParameters
+{
+	GENERATED_BODY()
+
+public:
+	UPROPERTY(EditAnywhere, BlueprintReadWrite, Category = "DynamicWind")
+	FVector SimulationCenter = FVector::Zero();
+
+	UPROPERTY(EditAnywhere, BlueprintReadWrite, Category = "DynamicWind")
+	float SimulationExtents = 1.0f;
+
+	UPROPERTY(EditAnywhere, BlueprintReadWrite, Category = "DynamicWind")
+	FVector WindDirection = FVector::ForwardVector;
+
+	UPROPERTY(EditAnywhere, BlueprintReadWrite, Category = "DynamicWind")
+	float WindSpeed = 15.0f;
+
+	UPROPERTY(EditAnywhere, BlueprintReadWrite, Category = "DynamicWind")
+	float WindAmplitude = 1.0f;
+
+	UPROPERTY(EditAnywhere, BlueprintReadWrite, Category = "DynamicWind")
+	TObjectPtr<UTexture> WindTexture = nullptr;
+
+	UPROPERTY(EditAnywhere, BlueprintReadWrite, Category = "DynamicWind")
+	FVector4f DebugModulation = FVector4f::Zero();
+};
diff --git a/Engine/Plugins/Experimental/DynamicWind/Source/DynamicWind/Public/DynamicWindProvider.h b/Engine/Plugins/Experimental/DynamicWind/Source/DynamicWind/Public/DynamicWindProvider.h
new file mode 100644
--- /dev/null
+++ b/Engine/Plugins/Experimental/DynamicWind/Source/DynamicWind/Public/DynamicWindProvider.h
@@ -0,0 +1,65 @@
+// Copyright Epic Games, Inc. All Rights Reserved.
+
+#pragma once
+
+#include "Engine/AssetUserData.h"
+#include "Engine/SkeletalMesh.h"
+#include "HLSLTypeAliases.h"
+#include "Shared/DynamicWindCommon.ush"
+#include "DynamicWindParameters.h"
+#include "Skinning/SkinningTransformProvider.h"
+#include "SkinningDefinitions.h"
+#include "Animation/TransformProviderData.h"
+
+class USkeleton;
+namespace Nanite { class FSkinnedSceneProxy; }
+
+#define DYNAMIC_WIND_TRANSFORM_PROVIDER_GUID 0xDFD4874B, 0xEE57466D, 0x883D6419, 0xAFE99EAC
+#define DYNAMIC_WIND_DEBUG_SKELETON_NAMES (UE_BUILD_DEBUG | UE_BUILD_DEVELOPMENT)
+
+class DYNAMICWIND_API FDynamicWindTransformProvider
+{
+public:
+	explicit FDynamicWindTransformProvider(FScene& InScene);
+	FDynamicWindTransformProvider() = delete;
+	~FDynamicWindTransformProvider();
+	
+	void RegisterSceneProxy(const Nanite::FSkinnedSceneProxy* InProxy);
+	void UnregisterSceneProxy(const Nanite::FSkinnedSceneProxy* InProxy);
+
+	int32 UpdateParameters(const FDynamicWindParameters& Parameters);
+
+	float GetBlendedWindAmplitude() const;
+
+private:
+	using FBoneDataBuffer = TPersistentByteAddressBuffer<FDynamicWindBoneData>;
+	using FBoneDataUploader = TByteAddressBufferScatterUploader<FDynamicWindBoneData>;
+
+	struct FSkeletonEntry
+	{	
+		uint64 UserDataHash = 0;
+		uint32 NumBones = 0;
+		uint32 ReferenceCount = 0;
+	#if DYNAMIC_WIND_DEBUG_SKELETON_NAMES
+		FName SkeletonName;
+	#endif
+		
+		FDynamicWindSkeletonData Data;
+	};
+
+	void ProvideTransforms(FSkinningTransformProvider::FProviderContext& Context);
+
+private:
+	FScene& Scene;
+
+	TMap<FGuid, FSkeletonEntry> SkeletonLookup;
+	FSpanAllocator BoneDataAllocator;
+	TUniquePtr<FBoneDataUploader> BoneDataUploader;
+	FBoneDataBuffer BoneDataBuffer;
+
+	FDynamicWindParameters WindParameters = {};
+
+	uint64 LastSimulatedFrameNumber = 0u;
+
+	float BlendedWindAmplitude = -1.0f;
+};
diff --git a/Engine/Plugins/Experimental/DynamicWind/Source/DynamicWind/Public/DynamicWindSkeletalData.h b/Engine/Plugins/Experimental/DynamicWind/Source/DynamicWind/Public/DynamicWindSkeletalData.h
new file mode 100644
--- /dev/null
+++ b/Engine/Plugins/Experimental/DynamicWind/Source/DynamicWind/Public/DynamicWindSkeletalData.h
@@ -0,0 +1,115 @@
+// Copyright Epic Games, Inc. All Rights Reserved.
+
+#pragma once
+
+#include "Engine/AssetUserData.h"
+
+#include "DynamicWindSkeletalData.generated.h"
+
+USTRUCT()
+struct FDynamicWindSimulationGroupData
+{
+	GENERATED_BODY()
+
+	UPROPERTY(EditAnywhere, Category = "Skeletal Data")
+	bool bUseDualInfluence = false;
+	
+	UPROPERTY(EditAnywhere, Category = "Skeletal Data", meta=(EditCondition = "!bUseDualInfluence", EditConditionHides, ClampMin=0.0f, ClampMax=1.0f))
+	float Influence = 0.0f;
+	
+	UPROPERTY(EditAnywhere, Category = "Skeletal Data", meta=(EditCondition = "bUseDualInfluence", EditConditionHides, ClampMin=0.0f, ClampMax=1.0f))
+	float MinInfluence = 0.0f;
+
+	UPROPERTY(EditAnywhere, Category = "Skeletal Data", meta=(EditCondition = "bUseDualInfluence", EditConditionHides, ClampMin=0.0f, ClampMax=1.0f))
+	float MaxInfluence = 0.0f;
+	
+	UPROPERTY(EditAnywhere, Category = "Skeletal Data", meta=(EditCondition = "bUseDualInfluence", EditConditionHides, ClampMin=0.0f, ClampMax=1.0f))
+	float ShiftTop = 0.0f;
+	
+	UPROPERTY(EditAnywhere, Category = "Skeletal Data", meta=(DisplayName = "Is Trunk?"))
+	bool bIsTrunkGroup = false;
+};
+
+USTRUCT()
+struct FDynamicWindSimulationGroupBoneLookup
+{
+	GENERATED_BODY()
+
+	UPROPERTY()
+	int32 SimulationGroupIndex = INDEX_NONE;
+
+	UPROPERTY()
+	TSet<int32> BoneIndices;
+};
+
+USTRUCT()
+struct FDynamicWindBoneChainData
+{
+	GENERATED_BODY()
+
+	UPROPERTY()
+	uint32 NumBones = 0;
+
+	UPROPERTY()
+	float ChainLength = 0.0f;
+};
+
+USTRUCT()
+struct FDynamicWindExtraBoneData
+{
+	GENERATED_BODY()
+
+	UPROPERTY()
+	int32 BoneChainOriginBoneIndex = INDEX_NONE;
+	
+	UPROPERTY()
+	int32 IndexInBoneChain = 0;
+};
+
+UCLASS()
+class DYNAMICWIND_API UDynamicWindSkeletalData : public UAssetUserData
+{
+	GENERATED_BODY()
+
+public:
+	/** When true, the SkeletalMesh object will be registered with the DynamicWind system. */
+	UPROPERTY(EditAnywhere, Category = "Skeletal Data")
+	bool bIsEnabled = false;
+
+	/** If true, marks this object as ground cover, and won't attenuate wind based on bone height. */
+	UPROPERTY(EditAnywhere, Category = "Skeletal Data")
+	bool bIsGroundCover = false;
+
+	/** Uniformly attenuates the influence of wind on trunk simulation groups. */
+	UPROPERTY(EditAnywhere, Category = "Skeletal Data")
+	float GustAttenuation = 0.0f;
+	
+	UPROPERTY(EditAnywhere, Category = "Skeletal Data")
+	TArray<FDynamicWindSimulationGroupData> SimulationGroups;
+	
+	UPROPERTY()
+	TArray<FDynamicWindSimulationGroupBoneLookup> SimulationGroupBones;
+
+	UPROPERTY()
+	TMap<int32, FDynamicWindBoneChainData> BoneChains;
+
+	UPROPERTY()
+	TMap<int32, FDynamicWindExtraBoneData> ExtraBonesData;
+
+public:
+	int32 GetBoneSimulationGroupIndex(int32 BoneIndex) const;
+	float GetBoneInfluence(int32 BoneIndex) const;
+	bool IsTrunkBone(int32 BoneIndex) const;
+	float GetBoneChainLength(int32 BoneIndex) const;
+	void RecalculateSkeletalDataHash();
+	uint64 GetSkeletalDataHash() const { return SkeletalDataHash; }
+	
+public:
+	virtual void PostLoad() override;
+#if WITH_EDITOR
+	virtual void PostEditChangeProperty(FPropertyChangedEvent& Event) override;
+#endif
+
+private:	
+	uint64 SkeletalDataHash = 0ull;
+};
diff --git a/Engine/Plugins/Experimental/DynamicWind/Source/DynamicWind/Public/DynamicWindSubsystem.h b/Engine/Plugins/Experimental/DynamicWind/Source/DynamicWind/Public/DynamicWindSubsystem.h
new file mode 100644
--- /dev/null
+++ b/Engine/Plugins/Experimental/DynamicWind/Source/DynamicWind/Public/DynamicWindSubsystem.h
@@ -0,0 +1,48 @@
+// Copyright Epic Games, Inc. All Rights Reserved.
+
+#pragma once
+
+#include "Subsystems/WorldSubsystem.h"
+#include "SpanAllocator.h"
+#include "Matrix3x4.h"
+#include "Containers/Map.h"
+#include "DynamicWindLog.h"
+#include "DynamicWindParameters.h"
+#include "DynamicWindProvider.h"
+#include "DynamicWindSubsystem.generated.h"
+
+class USkeleton;
+class UInstancedSkinnedMeshComponent;
+
+UCLASS(BlueprintType, ClassGroup = (Rendering, Common))
+class DYNAMICWIND_API UDynamicWindSubsystem : public UWorldSubsystem
+{
+	GENERATED_BODY()
+
+public:
+	UDynamicWindSubsystem();
+
+	virtual bool ShouldCreateSubsystem(UObject* Outer) const override;
+
+	virtual void PostInitialize() override;
+	virtual void PreDeinitialize() override;
+
+	void RegisterSceneProxy(const Nanite::FSkinnedSceneProxy* InSkinnedProxy);
+	void UnregisterSceneProxy(const Nanite::FSkinnedSceneProxy* InSkinnedProxy);
+
+public:
+	UFUNCTION(BlueprintCallable, Category = "DynamicWind")
+	float GetBlendedWindAmplitude() const;
+
+	UFUNCTION(BlueprintCallable, Category = "DynamicWind")
+	void UpdateWindParameters(const FDynamicWindParameters& Parameters);
+
+private:
+	bool bInitialized = false;
+
+	TUniquePtr<FDynamicWindTransformProvider> TransformProvider;
+	FGuid TransformProviderId;
+
+	FDelegateHandle OnCreateRenderThreadResources;
+	FDelegateHandle OnDestroyRenderThreadResources;
+};
diff --git a/Engine/Plugins/Experimental/DynamicWind/Source/DynamicWindEditor/DynamicWindEditor.Build.cs b/Engine/Plugins/Experimental/DynamicWind/Source/DynamicWindEditor/DynamicWindEditor.Build.cs
new file mode 100644
--- /dev/null
+++ b/Engine/Plugins/Experimental/DynamicWind/Source/DynamicWindEditor/DynamicWindEditor.Build.cs
@@ -0,0 +1,61 @@
+// Copyright Epic Games, Inc. All Rights Reserved.
+
+using System.IO;
+using UnrealBuildTool;
+
+namespace UnrealBuildTool.Rules
+{
+	public class DynamicWindEditor : ModuleRules
+	{
+		public DynamicWindEditor(ReadOnlyTargetRules Target) : base(Target)
+		{
+			PrivateIncludePaths.AddRange(
+				new string[] {
+					Path.Combine(GetModuleDirectory("Renderer"), "Private"),
+					Path.Combine(PluginDirectory, "Shaders")
+				}
+			);
+
+			PrivateDependencyModuleNames.AddRange(
+				new string[]
+				{
+					"AssetRegistry",
+					"AssetTools",
+					"Core",
+					"CoreUObject",
+					"EditorFramework",
+					"EditorStyle",
+					"Engine",
+					"PropertyEditor",
+					"RHI",
+					"Slate",
+					"Slate",
+					"SlateCore",
+					"TargetPlatform",
+					"UnrealEd",
+					"SourceControl",
+					"MeshDescription",
+					"ImageCore",
+					"StaticMeshDescription",
+					"SkeletalMeshDescription",
+					"Json",
+					"JsonUtilities"
+				}
+			);
+
+			PublicDependencyModuleNames.AddRange(
+				new string[]
+				{
+					"EditorSubsystem",
+					"DynamicWind",
+				}
+			);
+
+			PrivateIncludePathModuleNames.AddRange(
+				new string[] {
+					"DerivedDataCache",
+				}
+			);
+		}
+	}
+}
diff --git a/Engine/Plugins/Experimental/DynamicWind/Source/DynamicWindEditor/Private/DynamicWindBlueprintLibrary.cpp b/Engine/Plugins/Experimental/DynamicWind/Source/DynamicWindEditor/Private/DynamicWindBlueprintLibrary.cpp
new file mode 100644
--- /dev/null
+++ b/Engine/Plugins/Experimental/DynamicWind/Source/DynamicWindEditor/Private/DynamicWindBlueprintLibrary.cpp
@@ -0,0 +1,384 @@
+﻿// Copyright Epic Games, Inc. All Rights Reserved.
+
+#include "DynamicWindBlueprintLibrary.h"
+
+#include "IAssetTools.h"
+#include "ImageCore.h"
+#include "IntVectorTypes.h"
+#include "SkeletalMeshAttributes.h"
+#include "SkeletalMeshTypes.h"
+#include "StaticMeshAttributes.h"
+#include "Animation/Skeleton.h"
+#include "Engine/SkeletalMesh.h"
+#include "Engine/SkinnedAssetCommon.h"
+#include "Engine/StaticMesh.h"
+#include "Engine/Texture2D.h"
+#include "Rendering/SkeletalMeshModel.h"
+#include "DynamicWindSkeletalData.h"
+#include "DynamicWindImportData.h"
+#include "Misc/FileHelper.h"
+#include "JsonObjectConverter.h"
+#include "EditorFramework/AssetImportData.h"
+#include "DesktopPlatformModule.h"
+#include "Framework/Application/SlateApplication.h"
+
+DEFINE_LOG_CATEGORY_STATIC(LogDynamicWindBlueprintLibrary, Log, All);
+
+bool UDynamicWindBlueprintLibrary::ConvertPivotPainterTreeToSkeletalMesh(
+	const UStaticMesh* TreeStaticMesh,
+	UTexture2D* TreePivotPosTexture,
+	int32 TreePivotUVIndex,
+	USkeletalMesh* TargetSkeletalMesh,
+	USkeleton* TargetSkeleton
+	)
+{
+	if (!TreeStaticMesh)
+	{
+		UE_LOG(LogDynamicWindBlueprintLibrary, Error, TEXT("No static mesh given"));
+		return false;
+	}
+	if (!TreePivotPosTexture)
+	{
+		UE_LOG(LogDynamicWindBlueprintLibrary, Error, TEXT("No pivot pos texture given"));
+		return false;
+	}
+	if (!TargetSkeletalMesh)
+	{
+		UE_LOG(LogDynamicWindBlueprintLibrary, Error, TEXT("No target skeletal mesh given"));
+		return false;
+	}
+	if (!TargetSkeleton)
+	{
+		UE_LOG(LogDynamicWindBlueprintLibrary, Error, TEXT("No target skeleton given"));
+		return false;
+	}
+	if (TargetSkeletalMesh->GetLODNum() != 0 || TargetSkeletalMesh->GetSkeleton() || TargetSkeletalMesh->GetRefSkeleton().GetNum() != 0)
+	{
+		UE_LOG(LogDynamicWindBlueprintLibrary, Error, TEXT("Target skeletal mesh is not empty"));
+		return false;
+	}
+	if (TargetSkeleton->GetReferenceSkeleton().GetNum() != 0)
+	{
+		UE_LOG(LogDynamicWindBlueprintLibrary, Error, TEXT("Target skeleton is not empty"));
+		return false;
+	}
+	
+
+	// Try hi-res mesh first, fall back to LOD 0.
+	FMeshDescription MeshDescription; 
+	if (!TreeStaticMesh->CloneHiResMeshDescription(MeshDescription))
+	{
+		if (!TreeStaticMesh->CloneMeshDescription(0, MeshDescription))
+		{
+			UE_LOG(LogDynamicWindBlueprintLibrary, Error, TEXT("No valid geometry found on tree mesh"));
+			return false;
+		}
+	}
+
+	FSkeletalMeshAttributes MeshAttributes(MeshDescription);
+
+	constexpr bool bKeepExistingAttributes = true;
+	MeshAttributes.Register(bKeepExistingAttributes);
+	
+	TVertexInstanceAttributesConstRef<FVector2f> UVCoords = MeshAttributes.GetVertexInstanceUVs();
+	
+	if (TreePivotUVIndex < 0 || TreePivotUVIndex >= UVCoords.GetNumChannels())
+	{
+		UE_LOG(LogDynamicWindBlueprintLibrary, Error, TEXT("Invalid Pivot UV layer %d (Max layers: %d)"), TreePivotUVIndex, UVCoords.GetNumChannels());
+	}
+
+	FImage PivotImage;
+	TreePivotPosTexture->Source.GetMipImage(PivotImage, 0);
+
+	// Ask about which pixels are used from the pivot texture. Then use that to build the skeleton.
+	struct FPivot
+	{
+		FVector3f Pos;
+		int32 Parent = INDEX_NONE;
+		int32 BoneIndex = INDEX_NONE;
+		bool bHasChildren = false;
+		FName BoneName;
+	};
+	
+	TArray<FTriangleID> InvalidTriangleIDs;
+	TMap<int32, FPivot> Pivots;
+	TMap<FVector2f, int32> UVToPivot;
+	for (FTriangleID TriangleID: MeshDescription.Triangles().GetElementIDs())
+	{
+		TArrayView<const FVertexInstanceID> Corners = MeshDescription.GetTriangleVertexInstances(TriangleID);
+		for (int32 Index = 0; Index < Corners.Num(); Index++)
+		{
+			FVector2f UV = UVCoords.Get(Corners[Index], TreePivotUVIndex);
+
+			// Out-of-bound UV indicates root binding, we handle that when we do binding. 
+			if (UV.X >= 0.0f && UV.X < 1.0f && UV.Y >= 0.0f && UV.Y < 1.0f)
+			{
+				FInt32Vector2 Pixel;
+				Pixel.X = UV.X * PivotImage.SizeX;
+				Pixel.Y = UV.Y * PivotImage.SizeY;
+
+				const int32 PivotIndex = Pixel.X + Pixel.Y * PivotImage.SizeX;
+				Pivots.Add(PivotIndex, {});
+				UVToPivot.Add(UV, PivotIndex);
+			}
+			else
+			{
+				Pivots.Add(0, {});
+				UVToPivot.Add(UV, 0);
+			}
+		}
+	}
+
+	// Crawl through the pivot texture, using the used pivots, to get position and parent indices.
+	TArray<int32> PivotOrder;
+	
+	for (TPair<int32, FPivot>& Item: Pivots)
+	{
+		const int32 PivotIndex = Item.Key; 
+		FPivot& Pivot = Item.Value;
+		
+		PivotOrder.Add(PivotIndex);
+
+		// Pivot 0 is the root.
+		if (PivotIndex == 0)
+		{
+			Pivot.Parent = INDEX_NONE;
+			Pivot.Pos = FVector3f::ZeroVector;
+		}
+		else
+		{
+			const int32 X = PivotIndex % PivotImage.SizeX;
+			const int32 Y = PivotIndex / PivotImage.SizeX;
+			
+			FLinearColor C = PivotImage.GetOnePixelLinear(X, Y);
+			
+			const uint32 FloatBits = *reinterpret_cast<uint32*>(&C.A);
+			const uint32 Sign = ((FloatBits>>16)&0x8000);
+			const uint32 Exponent  = ((static_cast<const int32>((FloatBits >> 23) & 0xff)-127+15) << 10);
+			const uint32 Mantissa = ((FloatBits>>13)&0x3ff);
+			const uint32 Bits = (Sign | Exponent | Mantissa);
+			const uint32 ParentIndex = Bits - 1024;
+
+			Pivot.Pos = FVector3f(C.R, C.G, C.B);
+			Pivot.Parent = ParentIndex == PivotIndex ? 0 : ParentIndex;
+		}
+	}
+
+	// If nothing was explicitly bound to the root, add the root pivot. 
+	if(!Pivots.Contains(0))
+	{
+		FPivot RootPivot;
+		RootPivot.Parent = INDEX_NONE;
+		RootPivot.Pos = FVector3f::ZeroVector;
+		Pivots.Add(0, RootPivot);
+		PivotOrder.Add(0);
+	}
+
+	// Mark all the pivots that have children as such. We'll use this for naming.
+	for (TPair<int32, FPivot>& Item: Pivots)
+	{
+		const int32 ParentIndex = Item.Value.Parent;
+		if (Pivots.Contains(ParentIndex))
+		{
+			Pivots[ParentIndex].bHasChildren = true;
+		}
+	}
+
+	PivotOrder.Sort([Pivots](const int32 IndexA, const int32 IndexB)
+	{
+		return Pivots[IndexA].Parent < Pivots[IndexB].Parent;
+	});
+
+	// Give the pivots bone names and adjust the positions so that they're local to their parent.
+	int32 LeafId = 0, BranchId = 0, TrunkId = 0;
+	for (const int32 PivotIndex: PivotOrder)
+	{
+		FPivot& Pivot = Pivots[PivotIndex];
+
+		if (Pivot.Parent == INDEX_NONE)
+		{
+			Pivot.BoneName = FName("Root");
+		}
+		else if (Pivot.Parent == 0)
+		{
+			Pivot.BoneName = FName("Trunk", TrunkId++);
+		}
+		else if (Pivot.bHasChildren)
+		{
+			Pivot.BoneName = FName("Branch", BranchId++);
+		}
+		else
+		{
+			Pivot.BoneName = FName("Leaf", LeafId++);
+		}
+		
+		int32 ParentIndex = Pivot.Parent;
+		while (ParentIndex != INDEX_NONE)
+		{
+			Pivot.Pos -= Pivots[ParentIndex].Pos;
+			ParentIndex = Pivots[ParentIndex].Parent;
+		}
+	}
+	
+	// Construct the skeleton.
+	{
+		// This is a scoped operation. 
+		FReferenceSkeletonModifier SkeletonModifier(TargetSkeleton);
+
+		for (const int32 PivotIndex: PivotOrder)
+		{
+			FPivot& Pivot = Pivots[PivotIndex];
+
+			FMeshBoneInfo BoneInfo;
+			BoneInfo.Name = Pivot.BoneName;
+			BoneInfo.ExportName = Pivot.BoneName.ToString();
+			if (Pivot.Parent != INDEX_NONE)
+			{
+				BoneInfo.ParentIndex = Pivots[Pivot.Parent].BoneIndex;
+			}
+
+			FTransform BonePose(FVector3d(Pivot.Pos));
+		
+			SkeletonModifier.Add(BoneInfo, BonePose);
+			Pivot.BoneIndex = SkeletonModifier.FindBoneIndex(Pivot.BoneName);
+		}
+	}
+
+	// Bind the vertices as needed. If multiple instances of the same vertex refer to the different UVs,
+	// use that to smooth bind.
+	FSkinWeightsVertexAttributesRef SkinWeights = MeshAttributes.GetVertexSkinWeights();
+	TMap<FBoneIndexType, int32> BoneIndexCounts; 
+	for (FVertexID VertexID: MeshDescription.Vertices().GetElementIDs())
+	{
+		BoneIndexCounts.Reset();
+		for (FVertexInstanceID VertexInstanceID: MeshDescription.GetVertexVertexInstanceIDs(VertexID))
+		{
+			FVector2f UV = UVCoords.Get(VertexInstanceID, TreePivotUVIndex);
+			const int32 PivotIndex = UVToPivot[UV];
+			const FBoneIndexType BoneIndex = Pivots[PivotIndex].BoneIndex;
+
+			if (BoneIndexCounts.Contains(BoneIndex))
+			{
+				BoneIndexCounts[BoneIndex]++;
+			}
+			else
+			{
+				BoneIndexCounts.Add(BoneIndex, 1);
+			}
+		}
+
+		FBoneIndexType BoneIndexes[3];
+		float BoneWeights[3];
+
+		int32 Index = 0;
+		for (const TTuple<FBoneIndexType, int32>& Item: BoneIndexCounts)
+		{
+			BoneIndexes[Index] = Item.Key;
+			BoneWeights[Index] = Item.Value / 3.0f;
+			Index++;
+		}
+
+		SkinWeights.Set(VertexID, UE::AnimationCore::FBoneWeights::Create(BoneIndexes, BoneWeights, BoneIndexCounts.Num()));  
+	}
+	
+	FElementIDRemappings Remappings;
+	MeshDescription.Compact(Remappings);
+	
+	// Copy the static mesh over and bind the skeletal mesh to the new ref skeleton.
+	{
+		FScopedSkeletalMeshPostEditChange ScopedSkeletalMeshChange(TargetSkeletalMesh);
+
+		TargetSkeletalMesh->PreEditChange( nullptr );
+		TargetSkeletalMesh->SetRefSkeleton(TargetSkeleton->GetReferenceSkeleton());
+		TargetSkeletalMesh->CalculateInvRefMatrices();
+
+		TargetSkeletalMesh->SetNaniteSettings(TreeStaticMesh->GetNaniteSettings());
+
+		FSkeletalMeshModel* ImportedModels = TargetSkeletalMesh->GetImportedModel();
+		ImportedModels->LODModels.Add(new FSkeletalMeshLODModel);
+		
+		FSkeletalMeshLODInfo& LODInfo = TargetSkeletalMesh->AddLODInfo();
+		
+		// Make sure there's no reduction active.
+		LODInfo.ReductionSettings.NumOfTrianglesPercentage = 1.0f;
+		LODInfo.ReductionSettings.MaxNumOfTriangles = MAX_uint32;
+		LODInfo.ReductionSettings.MaxNumOfTrianglesPercentage = MAX_uint32;
+		LODInfo.ReductionSettings.NumOfVertPercentage = 1.0f;
+		LODInfo.ReductionSettings.MaxNumOfVerts = MAX_uint32;
+		LODInfo.ReductionSettings.MaxNumOfVertsPercentage = MAX_uint32;
+		LODInfo.ReductionSettings.MaxDeviationPercentage = 0.0f;
+
+		TargetSkeletalMesh->CreateMeshDescription(0, MoveTemp(MeshDescription));
+		TargetSkeletalMesh->CommitMeshDescription(0);
+
+		// Copy the materials over.
+		TArray<FSkeletalMaterial> Materials;
+		for (const FStaticMaterial& StaticMaterial: TreeStaticMesh->GetStaticMaterials())
+		{
+			FSkeletalMaterial Material(
+				StaticMaterial.MaterialInterface,
+				StaticMaterial.MaterialSlotName);
+		
+			Materials.Add(Material);
+		}
+
+		TargetSkeletalMesh->SetMaterials(Materials);
+
+		TargetSkeletalMesh->SetPositiveBoundsExtension(TreeStaticMesh->GetPositiveBoundsExtension());
+		TargetSkeletalMesh->SetNegativeBoundsExtension(TreeStaticMesh->GetNegativeBoundsExtension());
+	}
+
+	TargetSkeletalMesh->SetSkeleton(TargetSkeleton);
+	TargetSkeleton->MergeAllBonesToBoneTree(TargetSkeletalMesh);
+	TargetSkeleton->SetPreviewMesh(TargetSkeletalMesh);
+	return true;
+}
+
+bool UDynamicWindBlueprintLibrary::ImportDynamicWindSkeletalDataFromFile(USkeletalMesh* TargetSkeletalMesh)
+{
+	if (TargetSkeletalMesh == nullptr)
+	{
+		return false;
+	}
+
+	FString DefaultDirectory;
+	if (auto AssetImportData = TargetSkeletalMesh->GetAssetImportData())
+	{
+		DefaultDirectory = FPaths::GetPath(AssetImportData->GetFirstFilename());
+	}
+	else
+	{
+		DefaultDirectory = FPaths::ProjectContentDir();
+	}
+
+	TArray<FString> FileNames;
+	if (FDesktopPlatformModule::Get()->OpenFileDialog(
+		FSlateApplication::Get().FindBestParentWindowHandleForDialogs(nullptr),
+		TEXT("Import Dynamic Wind JSON"),
+		DefaultDirectory,
+		TEXT(""),
+		TEXT("JSON File (*.json)|*.json"),
+		EFileDialogFlags::None,
+		FileNames))
+	{
+		FString JsonString;
+		if (FFileHelper::LoadFileToString(JsonString, *FileNames[0]))
+		{
+			FDynamicWindSkeletalImportData ImportData;
+			if (FJsonObjectConverter::JsonObjectStringToUStruct(JsonString, &ImportData))
+			{
+				if (DynamicWind::ImportSkeletalData(*TargetSkeletalMesh, ImportData))
+				{
+					return true;
+				}
+			}
+		}
+
+		UE_LOG(LogDynamicWindBlueprintLibrary, Error,
+			TEXT("ImportDynamicWindSkeletalDataFromFile: Failed to load dynamic wind skeletal data from %s for SkeletalMesh %s."),
+			*FileNames[0], *TargetSkeletalMesh->GetFullName()
+			);
+	}
+
+	return false;
+}
diff --git a/Engine/Plugins/Experimental/DynamicWind/Source/DynamicWindEditor/Private/DynamicWindBlueprintLibrary.h b/Engine/Plugins/Experimental/DynamicWind/Source/DynamicWindEditor/Private/DynamicWindBlueprintLibrary.h
new file mode 100644
--- /dev/null
+++ b/Engine/Plugins/Experimental/DynamicWind/Source/DynamicWindEditor/Private/DynamicWindBlueprintLibrary.h
@@ -0,0 +1,29 @@
+﻿// Copyright Epic Games, Inc. All Rights Reserved.
+
+#pragma once
+
+#include "DynamicWindImportData.h"
+
+#include "DynamicWindBlueprintLibrary.generated.h"
+
+class USkeletalMesh;
+class UStaticMesh;
+class UTexture2D;
+
+UCLASS()
+class UDynamicWindBlueprintLibrary : public UObject
+{
+	GENERATED_BODY()
+
+	UFUNCTION(BlueprintCallable, Category="DynamicWind")
+	static bool ConvertPivotPainterTreeToSkeletalMesh(
+		const UStaticMesh* TreeStaticMesh,
+		UTexture2D* TreePivotPosTexture,
+		int32 TreePivotUVIndex,
+		USkeletalMesh* TargetSkeletalMesh,
+		USkeleton* TargetSkeleton
+	);
+
+	UFUNCTION(BlueprintCallable, Category="DynamicWind")
+	static bool ImportDynamicWindSkeletalDataFromFile(USkeletalMesh* TargetSkeletalMesh);
+};
diff --git a/Engine/Plugins/Experimental/DynamicWind/Source/DynamicWindEditor/Private/DynamicWindEditorModule.cpp b/Engine/Plugins/Experimental/DynamicWind/Source/DynamicWindEditor/Private/DynamicWindEditorModule.cpp
new file mode 100644
--- /dev/null
+++ b/Engine/Plugins/Experimental/DynamicWind/Source/DynamicWindEditor/Private/DynamicWindEditorModule.cpp
@@ -0,0 +1,32 @@
+// Copyright Epic Games, Inc. All Rights Reserved.
+
+#include "DynamicWindEditorModule.h"
+#include "PropertyEditorModule.h"
+#include "Modules/ModuleManager.h"
+
+#define LOCTEXT_NAMESPACE "DynamicWindEditor"
+
+void FDynamicWindEditorModule::StartupModule()
+{
+	//FAssetToolsModule& AssetToolsModule = FAssetToolsModule::GetModule();
+	//IAssetTools& AssetTools = AssetToolsModule.Get();
+
+	FPropertyEditorModule& PropertyModule = FModuleManager::LoadModuleChecked<FPropertyEditorModule>("PropertyEditor");
+
+}
+	
+void FDynamicWindEditorModule::ShutdownModule()
+{
+	FPropertyEditorModule& PropertyModule = FModuleManager::LoadModuleChecked<FPropertyEditorModule>("PropertyEditor");
+
+}
+
+FDynamicWindEditorModule& FDynamicWindEditorModule::GetModule()
+{
+	static const FName ModuleName = "DynamicWindEditor";
+	return FModuleManager::LoadModuleChecked<FDynamicWindEditorModule>(ModuleName);
+}
+
+IMPLEMENT_MODULE(FDynamicWindEditorModule, DynamicWindEditor);
+
+#undef LOCTEXT_NAMESPACE
diff --git a/Engine/Plugins/Experimental/DynamicWind/Source/DynamicWindEditor/Private/DynamicWindEditorModule.h b/Engine/Plugins/Experimental/DynamicWind/Source/DynamicWindEditor/Private/DynamicWindEditorModule.h
new file mode 100644
--- /dev/null
+++ b/Engine/Plugins/Experimental/DynamicWind/Source/DynamicWindEditor/Private/DynamicWindEditorModule.h
@@ -0,0 +1,22 @@
+// Copyright Epic Games, Inc. All Rights Reserved.
+
+#pragma once
+
+#include "Modules/ModuleInterface.h"
+#include "UObject/StrongObjectPtr.h"
+
+class UPackage;
+class FString;
+
+class FDynamicWindEditorModule : public IModuleInterface
+{
+public:
+
+	virtual void StartupModule() override;
+	virtual void ShutdownModule() override;
+
+	static FDynamicWindEditorModule& GetModule();
+
+private:
+
+};
diff --git a/Engine/Plugins/Experimental/DynamicWind/Source/DynamicWindEditor/Private/DynamicWindFactory.cpp b/Engine/Plugins/Experimental/DynamicWind/Source/DynamicWindEditor/Private/DynamicWindFactory.cpp
new file mode 100644
--- /dev/null
+++ b/Engine/Plugins/Experimental/DynamicWind/Source/DynamicWindEditor/Private/DynamicWindFactory.cpp
@@ -0,0 +1,17 @@
+// Copyright Epic Games, Inc. All Rights Reserved.
+
+#include "DynamicWindFactory.h"
+#include "DynamicWindProvider.h"
+#include "DynamicWindData.h"
+
+UDynamicWindDataFactory::UDynamicWindDataFactory(const FObjectInitializer& ObjectInitializer)
+: Super(ObjectInitializer)
+{
+	SupportedClass = UDynamicWindData::StaticClass();
+	ProviderDataClass = UDynamicWindData::StaticClass();
+}
+
+bool UDynamicWindDataFactory::ConfigureProperties()
+{
+	return true;
+}
diff --git a/Engine/Plugins/Experimental/DynamicWind/Source/DynamicWindEditor/Private/DynamicWindImportData.cpp b/Engine/Plugins/Experimental/DynamicWind/Source/DynamicWindEditor/Private/DynamicWindImportData.cpp
new file mode 100644
--- /dev/null
+++ b/Engine/Plugins/Experimental/DynamicWind/Source/DynamicWindEditor/Private/DynamicWindImportData.cpp
@@ -0,0 +1,145 @@
+// Copyright Epic Games, Inc. All Rights Reserved.
+
+#include "DynamicWindImportData.h"
+#include "Engine/SkeletalMesh.h"
+
+namespace DynamicWind
+{
+
+UDynamicWindSkeletalData* ImportSkeletalData(
+	USkeletalMesh& TargetSkeletalMesh,
+	const FDynamicWindSkeletalImportData& ImportData
+)
+{
+	const FReferenceSkeleton& SkeletonRef = TargetSkeletalMesh.GetRefSkeleton();
+	const TArray<FMeshBoneInfo>& MeshBoneInfos = SkeletonRef.GetRawRefBoneInfo();
+
+	struct FBoneInfo
+	{
+		int32 ID = INDEX_NONE;
+		int32 ParentID = INDEX_NONE;
+		FName Name = NAME_None;
+	
+		int32 SimulationGroupIndex = -1;
+
+		FVector3f Start = FVector3f::Zero();
+		FVector3f End = FVector3f::Zero();
+	};
+
+	int32 MaxSimulationGroupIndex = -1;
+	TArray<FBoneInfo> BoneInfos;
+	BoneInfos.Reserve(MeshBoneInfos.Num());
+	for (auto Bone : MeshBoneInfos)
+	{
+		int Index = SkeletonRef.FindRawBoneIndex(Bone.Name);
+
+		if (Index != INDEX_NONE)
+		{
+			FBoneInfo NewBone;
+			NewBone.ID = Index;
+			NewBone.ParentID = Bone.ParentIndex;
+			NewBone.Name = Bone.Name;
+
+			int JointIndex = ImportData.Joints.IndexOfByPredicate(
+				[BoneName=Bone.Name] (const FDynamicWindJointImportData& Joint)
+				{
+					return Joint.JointName == BoneName;
+				}
+			);
+
+			if (JointIndex != INDEX_NONE)
+			{
+				const FDynamicWindJointImportData& Joint = ImportData.Joints[JointIndex];
+				NewBone.SimulationGroupIndex = Joint.SimulationGroupIndex;
+
+				MaxSimulationGroupIndex = FMath::Max(MaxSimulationGroupIndex, NewBone.SimulationGroupIndex);
+			}
+
+			const auto& RawRefBonePose = SkeletonRef.GetRawRefBonePose();
+			NewBone.Start = RawRefBonePose.IsValidIndex(Bone.ParentIndex) ?
+				FVector3f(RawRefBonePose[Bone.ParentIndex].GetLocation()) : FVector3f::Zero();
+			NewBone.End = FVector3f(RawRefBonePose[Index].GetLocation());
+
+			BoneInfos.Add(NewBone);
+		}
+	}
+
+	auto AssetUserData = TargetSkeletalMesh.GetAssetUserData<UDynamicWindSkeletalData>();
+	if (AssetUserData == nullptr)
+	{
+		AssetUserData = NewObject<UDynamicWindSkeletalData>(&TargetSkeletalMesh);
+		TargetSkeletalMesh.AddAssetUserData(AssetUserData);
+	}
+	else
+	{
+		// Reset the user data
+		AssetUserData->SimulationGroups.Reset();
+		AssetUserData->SimulationGroupBones.Reset();
+		AssetUserData->BoneChains.Reset();
+		AssetUserData->ExtraBonesData.Reset();
+	}
+
+	AssetUserData->bIsEnabled			= true;
+	AssetUserData->bIsGroundCover		= ImportData.bIsGroundCover;
+	AssetUserData->GustAttenuation		= ImportData.GustAttenuation;
+	AssetUserData->SimulationGroups		= ImportData.SimulationGroups;
+
+	// To be fault tolerant, add default data for simulation groups assigned to joints that are not represented
+	// in the import data
+	TMap<int32, TSet<int32>> SimulationGroupBones;
+	while (AssetUserData->SimulationGroups.Num() < MaxSimulationGroupIndex + 1)
+	{
+		AssetUserData->SimulationGroups.Add({
+			.Influence = 1.0f
+		});
+	}
+
+	for (const auto& BoneInfo : BoneInfos)
+	{
+		if (BoneInfo.SimulationGroupIndex < 0)
+		{
+			continue;
+		}
+
+		SimulationGroupBones.FindOrAdd(BoneInfo.SimulationGroupIndex).Add(BoneInfo.ID);
+
+		// Find bone chain origin bone and this bone's position in chain
+		FDynamicWindExtraBoneData ExtraBoneData
+		{
+			.BoneChainOriginBoneIndex = BoneInfo.ID
+		};
+		int32 CurrentParent = BoneInfo.ParentID;
+		while (CurrentParent != INDEX_NONE)
+		{
+			if (BoneInfos[CurrentParent].SimulationGroupIndex != BoneInfo.SimulationGroupIndex)
+			{
+				break;
+			}
+
+			ExtraBoneData.BoneChainOriginBoneIndex = BoneInfos[CurrentParent].ID;
+			ExtraBoneData.IndexInBoneChain++;
+			
+			CurrentParent = BoneInfos[CurrentParent].ParentID;
+		}
+		
+		auto& ChainInfo = AssetUserData->BoneChains.FindOrAdd(ExtraBoneData.BoneChainOriginBoneIndex);
+		ChainInfo.NumBones++;
+		ChainInfo.ChainLength += FVector3f::Distance(BoneInfo.Start, BoneInfo.End);
+		AssetUserData->ExtraBonesData.Emplace(BoneInfo.ID, ExtraBoneData);
+	}
+
+	for (auto& [SimulationGroupIndex, BoneIndices] : SimulationGroupBones)
+	{
+		AssetUserData->SimulationGroupBones.Add({ SimulationGroupIndex, MoveTemp(BoneIndices) });
+	}
+	
+	AssetUserData->RecalculateSkeletalDataHash();
+	
+	TargetSkeletalMesh.MarkPackageDirty();
+	TargetSkeletalMesh.PostEditChange();
+	TargetSkeletalMesh.Modify();
+
+	return AssetUserData;
+}
+
+} // namespace DynamicWind
\ No newline at end of file
diff --git a/Engine/Plugins/Experimental/DynamicWind/Source/DynamicWindEditor/Public/DynamicWindFactory.h b/Engine/Plugins/Experimental/DynamicWind/Source/DynamicWindEditor/Public/DynamicWindFactory.h
new file mode 100644
--- /dev/null
+++ b/Engine/Plugins/Experimental/DynamicWind/Source/DynamicWindEditor/Public/DynamicWindFactory.h
@@ -0,0 +1,18 @@
+// Copyright Epic Games, Inc. All Rights Reserved.
+
+#pragma once
+
+#include "CoreMinimal.h"
+#include "UObject/ObjectMacros.h"
+#include "Factories/Factory.h"
+#include "Factories/TransformProviderFactory.h"
+#include "DynamicWindFactory.generated.h"
+
+UCLASS(MinimalAPI)
+class UDynamicWindDataFactory : public UTransformProviderDataFactory
+{
+	GENERATED_UCLASS_BODY()
+
+	virtual bool ConfigureProperties() override;
+	virtual bool ShouldShowInNewMenu() const override { return false; }
+};
diff --git a/Engine/Plugins/Experimental/DynamicWind/Source/DynamicWindEditor/Public/DynamicWindImportData.h b/Engine/Plugins/Experimental/DynamicWind/Source/DynamicWindEditor/Public/DynamicWindImportData.h
new file mode 100644
--- /dev/null
+++ b/Engine/Plugins/Experimental/DynamicWind/Source/DynamicWindEditor/Public/DynamicWindImportData.h
@@ -0,0 +1,50 @@
+// Copyright Epic Games, Inc. All Rights Reserved.
+
+#pragma once
+
+#include "DynamicWindSkeletalData.h"
+
+#include "DynamicWindImportData.generated.h"
+
+class USkeletalMesh;
+
+USTRUCT()
+struct FDynamicWindJointImportData
+{
+	GENERATED_USTRUCT_BODY()
+
+	UPROPERTY(EditAnywhere, Category = "Default")
+	FName JointName;
+
+	UPROPERTY(EditAnywhere, Category = "Default")
+	int32 SimulationGroupIndex = 0;
+};
+
+USTRUCT(BlueprintType)
+struct FDynamicWindSkeletalImportData
+{
+	GENERATED_USTRUCT_BODY()
+
+	UPROPERTY(EditAnywhere, Category = "SkeletalData")
+	TArray<FDynamicWindJointImportData> Joints;
+
+	UPROPERTY(EditAnywhere, Category = "SkeletalData")
+	TArray<FDynamicWindSimulationGroupData> SimulationGroups;
+
+	UPROPERTY(EditAnywhere, Category = "SkeletalData")
+	bool bIsGroundCover = false;
+
+	UPROPERTY(EditAnywhere, Category = "SkeletalData")
+	float GustAttenuation = 0.0f;
+};
+
+namespace DynamicWind
+{
+
+/** Attempts to create dynamic wind asset user data for a skeletal mesh with the given import data */
+UDynamicWindSkeletalData* ImportSkeletalData(
+	USkeletalMesh& TargetSkeletalMesh,
+	const FDynamicWindSkeletalImportData& ImportData
+);
+
+}
