[DynamicWind] Add example python script to parse a USD file containing a skeleton with the Unreal DynamicWindSkeletonAPI schema applied to it and generate a JSON file that can be consumed by the asset action blueprint utility included in the Dynamic Wind plugin.

The generation of the JSON file from USD using this script is a temporary solution for importing dynamic wind data from USD until we have the support required to extend the importer to generate this data automatically.

#rnx

diff --git a/Engine/Plugins/Experimental/DynamicWind/Resources/PythonExamples/export_dynamic_wind_json.py b/Engine/Plugins/Experimental/DynamicWind/Resources/PythonExamples/export_dynamic_wind_json.py
new file mode 100644
--- /dev/null
+++ b/Engine/Plugins/Experimental/DynamicWind/Resources/PythonExamples/export_dynamic_wind_json.py
@@ -0,0 +1,130 @@
+# Copyright Epic Games, Inc. All Rights Reserved.
+
+import sys
+import argparse
+import json
+from pxr import Usd
+
+# global token values
+class Tokens:
+	DynamicWindSkeletonAPI = 'DynamicWindSkeletonAPI'
+
+	class Attributes:
+		JointNames = 'unreal:dynamicWind:jointNames'
+		JointSimulationGroups = 'unreal:dynamicWind:jointSimulationGroups'
+		SimulationGroupInfluences = 'unreal:dynamicWind:simulationGroupInfluences'
+		SimulationGroupNumInfluences = 'unreal:dynamicWind:simulationGroupNumInfluences'
+		SimulationGroupShiftTops = 'unreal:dynamicWind:simulationGroupShiftTops'
+		TrunkSimulationGroups = 'unreal:dynamicWind:trunkSimulationGroups'
+		IsGroundCover = 'unreal:dynamicWind:isGroundCover'
+		GustAttenuation = 'unreal:dynamicWind:gustAttenuation'
+
+def createJointData(jointName, simulationGroupIndex=0):
+	return {
+		'JointName': jointName,
+		'SimulationGroupIndex': simulationGroupIndex
+	}
+
+def createSingleInfluenceGroupData(influence, shiftTop=0.0, isTrunkGroup=False):
+	return {
+		'bUseDualInfluence': False,
+		'Influence': influence,
+		'MinInfluence': 0.0,
+		'MaxInfluence': 0.0,
+		'ShiftTop': shiftTop,
+		'bIsTrunkGroup': isTrunkGroup
+	}
+
+def createDualInfluenceGroupData(minInfluence, maxInfluence, shiftTop=0.0, isTrunkGroup=False):
+	return {
+		'bUseDualInfluence': True,
+		'Influence': 0.0,
+		'MinInfluence': minInfluence,
+		'MaxInfluence': maxInfluence,
+		'ShiftTop': shiftTop,
+		'bIsTrunkGroup': isTrunkGroup
+	}
+
+def createSkeletalImportData(joints, groups, isGroundCover=False, gustAttenuation=0.0):
+	return {
+		'Joints': joints,
+		'SimulationGroups': groups,
+		'bIsGroundCover': isGroundCover,
+		'GustAttenuation': gustAttenuation
+	}
+
+def loadImportDataFromUsdFile(stage, primPath):
+	skelPrim = stage.GetPrimAtPath(primPath)
+	if not skelPrim:
+		print(f'Prim "{primPath}" doesn\'t exist')
+		return None
+	
+	if not skelPrim.HasAPI(Tokens.DynamicWindSkeletonAPI):
+		appliedSchemas = skelPrim.GetAppliedSchemas()
+		print(f'Prim "{primPath}" doesn\'t have the DynamicWindSkeletonAPI schema applied to it. Schemas: {appliedSchemas}')
+		return None
+
+	jointNamesAttr = skelPrim.GetAttribute(Tokens.Attributes.JointNames)
+	simulationGroupsAttr = skelPrim.GetAttribute(Tokens.Attributes.JointSimulationGroups)
+	simulationGroupInfluencesAttr = skelPrim.GetAttribute(Tokens.Attributes.SimulationGroupInfluences)
+	simulationGroupNumInfluencesAttr = skelPrim.GetAttribute(Tokens.Attributes.SimulationGroupNumInfluences)
+	simulationGroupShiftTopsAttr = skelPrim.GetAttribute(Tokens.Attributes.SimulationGroupShiftTops)
+	trunkSimulationGroupsAttr = skelPrim.GetAttribute(Tokens.Attributes.TrunkSimulationGroups)
+	isGroundCoverAttr = skelPrim.GetAttribute(Tokens.Attributes.IsGroundCover)
+	gustAttenuationAttr = skelPrim.GetAttribute(Tokens.Attributes.GustAttenuation)
+
+	jointNames = jointNamesAttr.Get() or [] if jointNamesAttr else []
+	simulationGroups = simulationGroupsAttr.Get() or [] if simulationGroupsAttr else []
+	simulationGroupInfluences = simulationGroupInfluencesAttr.Get() or [] if simulationGroupInfluencesAttr else []
+	simulationGroupNumInfluences = simulationGroupNumInfluencesAttr.Get() or [] if simulationGroupNumInfluencesAttr else []
+	simulationGroupShiftTops = simulationGroupShiftTopsAttr.Get() or [] if simulationGroupShiftTopsAttr else []
+	trunkSimulationGroups = trunkSimulationGroupsAttr.Get() or [] if trunkSimulationGroupsAttr else []
+	isGroundCover = bool(isGroundCoverAttr.Get()) if isGroundCoverAttr else False
+	gustAttenuation = float(gustAttenuationAttr.Get()) if gustAttenuationAttr else 0.0
+
+	joints = [createJointData(j, s) for j, s in zip(jointNames, simulationGroups)]
+	groups = []
+
+	getInfluence = lambda i: simulationGroupInfluences[i] if i < len(simulationGroupInfluences) else 1.0
+	getShiftTop = lambda i: simulationGroupShiftTops[i] if i < len(simulationGroupShiftTops) else 0.0
+
+	if not simulationGroupNumInfluences:
+		# all single influence simulation groups, num groups = len(simulationGroupInfluences)
+		groups = [createSingleInfluenceGroupData(inf, getShiftTop(i)) for i, inf in enumerate(simulationGroupInfluences)]
+	else:
+		# num groups = len(simulationGroupNumInfluences)
+		infIdx = 0
+		for i, numInfluences in enumerate(simulationGroupNumInfluences):
+			if numInfluences == 2:
+				# dual influence group
+				groups.append(createDualInfluenceGroupData(getInfluence(infIdx), getInfluence(infIdx + 1), getShiftTop(i)))
+				infIdx += 2
+			else:
+				# single influence group
+				groups.append(createSingleInfluenceGroupData(getInfluence(infIdx), getShiftTop(i)))
+				infIdx += 1
+
+	# mark trunk groups
+	for i in trunkSimulationGroups:
+		if i >= 0 and i < len(groups):
+			groups[i]['isTrunkGroup'] = True
+
+	return createSkeletalImportData(joints, groups, isGroundCover, gustAttenuation)
+
+if __name__ == '__main__':
+	parser = argparse.ArgumentParser(
+		prog='ExportDynamicWindJson',
+		description='Generates an Unreal Dynamic Wind JSON file from USD'
+	)
+	parser.add_argument('usdFile', help="Input USD File")
+	parser.add_argument('jsonFile', help="Output JSON File")
+	parser.add_argument('-p', '--primPath', help='USD Prim path of the relevant skeleton')
+
+	args = parser.parse_args()
+
+	stage = Usd.Stage.Open(args.usdFile)
+	importData = loadImportDataFromUsdFile(stage, args.primPath)
+	if importData:
+		with open(args.jsonFile, 'w') as file:
+			jsonString = json.dumps(importData, indent=4)
+			file.write(jsonString)
\ No newline at end of file
