[DynamicWind] Fix a bug in the example dynamic wind JSON export script with the marking of trunk simulation groups

#rnx

diff --git a/Engine/Plugins/Experimental/DynamicWind/Resources/PythonExamples/export_dynamic_wind_json.py b/Engine/Plugins/Experimental/DynamicWind/Resources/PythonExamples/export_dynamic_wind_json.py
--- a/Engine/Plugins/Experimental/DynamicWind/Resources/PythonExamples/export_dynamic_wind_json.py
+++ b/Engine/Plugins/Experimental/DynamicWind/Resources/PythonExamples/export_dynamic_wind_json.py
@@ -107,7 +107,7 @@ def loadImportDataFromUsdFile(stage, primPath):
 	# mark trunk groups
 	for i in trunkSimulationGroups:
 		if i >= 0 and i < len(groups):
-			groups[i]['isTrunkGroup'] = True
+			groups[i]['bIsTrunkGroup'] = True
 
 	return createSkeletalImportData(joints, groups, isGroundCover, gustAttenuation)
 
