Implemented support for per-instance rotation working correctly with global wind direction in the dynamic wind plugin. This change includes a new API for the transform providers that allow for per-instance skinning data offset calculations.

diff --git a/Engine/Plugins/Experimental/DynamicWind/Shaders/DynamicWindEval.usf b/Engine/Plugins/Experimental/DynamicWind/Shaders/DynamicWindEval.usf
--- a/Engine/Plugins/Experimental/DynamicWind/Shaders/DynamicWindEval.usf
+++ b/Engine/Plugins/Experimental/DynamicWind/Shaders/DynamicWindEval.usf
@@ -54,9 +54,6 @@ float4 Debug;
 #define DYNAMIC_WIND_GUST_SPEED_MIN 0.05f
 #define DYNAMIC_WIND_GUST_SPEED_MAX 0.1f
 
-// TODO Use a cvar instead of a define
-#define DYNAMIC_WIND_DIRECTIONALITY_SLICES 8.0f
-
 uint GetBoneDataAddress(uint BoneDataOffset, uint BoneIndex)
 {
 	return (BoneDataOffset + BoneIndex) * uint(sizeof(FDynamicWindBoneData));
@@ -146,14 +143,14 @@ void WindEvalCS(uint3 GroupId : SV_GroupID, uint GroupThreadIndex : SV_GroupInde
 	float3 PositionAccumulation = float3(0.0f, 0.0f, 0.0f);
 	float3 LastBoneBindPosition = BoneData.BindPosePosition;
 
-	float SliceAngle = 2.0f * PI / DYNAMIC_WIND_DIRECTIONALITY_SLICES;
+	float SliceAngle = 2.0f * PI / float(DYNAMIC_WIND_DIRECTIONALITY_SLICES);
 
 	for (uint CurrentBoneDataAddress = BoneDataAddress, ParentBoneIndex = SkeletonBoneIndex; ParentBoneIndex != DYNAMIC_WIND_INVALID_BONE_INDEX;)
 	{
 		const FDynamicWindBoneData CurrentBoneData = BoneDataBuffer.Load<FDynamicWindBoneData>(CurrentBoneDataAddress);
 
 		// === ROTATION AXIS CALCULATIONS ========================================================================================================
-		const float3 SectionWindDirection = QuatRotateVector(QuatConjugate(QuatFromAxisAngle(float3(0.0f, 0.0f, 1.0f), Header.BlockAnimIndex * SliceAngle)), WindDirection);
+		const float3 SectionWindDirection = QuatRotateVector(QuatConjugate(QuatFromAxisAngle(float3(0.0f, 0.0f, 1.0f), Header.BlockDirectionalityIndex * SliceAngle)), WindDirection);
 		const float3 AdjustedBoneForward = QuatRotateVector(CurrentBoneData.BindPoseRotation, float3(1.0f, 0.0f, 0.0f));
 			
 		// Prevent error when the bone forward and wind direction are collinear
@@ -169,7 +166,7 @@ void WindEvalCS(uint3 GroupId : SV_GroupID, uint GroupThreadIndex : SV_GroupInde
 		const float HeightGrad = saturate(CurrentBoneData.BindPosePosition.z / SkeletonData.SkeletonHeight);
 		// RandBBSfloat(0) returns 0, which result in an uniform rotation for all branches
 		// SkeletonDataIndex is 0 for the first registered skeleton, hence the +1
-		const float RandomValue = RandBBSfloat(((Header.SkeletonDataIndex << 8) | Header.BlockAnimIndex) + 1);
+		const float RandomValue = RandBBSfloat(((Header.SkeletonDataIndex << 8) | Header.BlockDirectionalityIndex) + 1);
 			
 		float GustSpeed = lerp(DYNAMIC_WIND_GUST_SPEED_MIN, DYNAMIC_WIND_GUST_SPEED_MAX, WindSpeed / 100.0f);
 		float GustStrength = RampStrength(WindSpeed / 100.0f);
diff --git a/Engine/Plugins/Experimental/DynamicWind/Shaders/Shared/DynamicWindCommon.ush b/Engine/Plugins/Experimental/DynamicWind/Shaders/Shared/DynamicWindCommon.ush
--- a/Engine/Plugins/Experimental/DynamicWind/Shaders/Shared/DynamicWindCommon.ush
+++ b/Engine/Plugins/Experimental/DynamicWind/Shaders/Shared/DynamicWindCommon.ush
@@ -4,6 +4,8 @@
 
 #define DYNAMIC_WIND_INVALID_BONE_INDEX 65535
 
+#define DYNAMIC_WIND_DIRECTIONALITY_SLICES 8
+
 #ifdef __cplusplus
 namespace UE::HLSL
 {
@@ -13,12 +15,16 @@ struct FDynamicWindSkeletonData
 {
 	/** The index of the first element of the FDynamicWindBoneData for this skeleton */
 	uint BoneDataOffset;
+
 	/** The height (in units) of the skeleton */
 	float SkeletonHeight;
+
 	/** How much to attenuate trunk bones in the skeleton (0-1) */
 	float GustAttenuation;
+
 	/** The maximum SimulationGroupIndex value for all bones in the skeleton */
 	uint MaxSimulationGroupIndex : 31;
+
 	/** Whether or not this object is ground cover */
 	uint bIsGroundCover : 1;
 };
@@ -27,16 +33,22 @@ struct FDynamicWindBoneData
 {
 	/** The index of this bone's parent bone, in the skeleton */
 	uint ParentBoneIndex;
+
 	/** The wind influence of this bone */
 	float Influence;
+
 	/** The simulation group for this bone */
 	int SimulationGroupIndex;
+
 	/** The length of bone chain that this bone is a member of */
 	float BoneChainLength;
+
 	/** Initial local-space orientation of the bone */
 	float4 BindPoseRotation;
+
 	/** Initial local-space position of the bone */
 	float3 BindPosePosition;
+
 	/** Whether this bone is to be considered part of a tree trunk */
 	uint bIsTrunkBone;
 };
@@ -45,14 +57,19 @@ struct FDynamicWindBlockHeader
 {	
 	/** Index of the related FDynamicWindSkeletonData for the block */
 	uint SkeletonDataIndex;
+
 	/** Byte offset of the transforms to be written in the transform buffer */
 	uint BlockDstTransformOffset;
+
 	/** Number of bones in this block */
 	uint BlockNumBones : 8;
+
 	/** Index of the first bone (in the skeleton) of the block */
 	uint BlockBoneOffset : 24;
-	/** Index of the unique anim this block is associated with */
-	uint BlockAnimIndex : 8;
+
+	/** Index of the unique directional anim this block is associated with */
+	uint BlockDirectionalityIndex : 8;
+
 	/** Total number of transforms to be written for this skeleton anim. */
 	uint TotalTransformCount : 24;
 };
diff --git a/Engine/Plugins/Experimental/DynamicWind/Source/DynamicWind/Private/DynamicWindData.cpp b/Engine/Plugins/Experimental/DynamicWind/Source/DynamicWind/Private/DynamicWindData.cpp
--- a/Engine/Plugins/Experimental/DynamicWind/Source/DynamicWind/Private/DynamicWindData.cpp
+++ b/Engine/Plugins/Experimental/DynamicWind/Source/DynamicWind/Private/DynamicWindData.cpp
@@ -8,6 +8,7 @@
 #include "ScenePrivate.h"
 #include "SkeletalRenderPublic.h"
 #include "SkinningDefinitions.h"
+#include "Components/InstancedSkinnedMeshComponent.h"
 
 bool UDynamicWindData::IsEnabled() const
 {
@@ -27,7 +28,7 @@ bool UDynamicWindData::UsesSkeletonBatching() const
 
 const uint32 UDynamicWindData::GetUniqueAnimationCount() const
 {
-	return 1;
+	return DYNAMIC_WIND_DIRECTIONALITY_SLICES;
 }
 
 bool UDynamicWindData::HasAnimationBounds() const
@@ -40,6 +41,19 @@ bool UDynamicWindData::GetAnimationBounds(uint32 AnimationIndex, FRenderBounds&
 	return false;
 }
 
+uint32 UDynamicWindData::GetSkinningDataOffset(int32 InstanceIndex, const FTransform& ComponentTransform, const FSkinnedMeshInstanceData& InstanceData) const
+{
+	FTransform WorldSpaceTransform = FTransform(InstanceData.Transform) * ComponentTransform;
+	const FVector EulerRotation = WorldSpaceTransform.GetRotation().Euler();
+	int32 DirectionalityIndex = 0;
+	{
+		const int32 ZRoundDown = FMath::RoundToInt32(float(EulerRotation.Z) / 45.0f) * 45u;
+		DirectionalityIndex = ((ZRoundDown + 360u) % 360u) / 45u;
+	}
+
+	return DirectionalityIndex * 2u;
+}
+
 FTransformProviderRenderProxy* UDynamicWindData::CreateRenderThreadResources(FSkinningSceneExtensionProxy* SceneProxy, FSceneInterface& Scene, FRHICommandListBase& RHICmdList)
 {
 	FDynamicWindDataRenderProxy* ProviderProxy = new FDynamicWindDataRenderProxy(this, SceneProxy, Scene);
@@ -84,4 +98,4 @@ const TConstArrayView<uint64> FDynamicWindDataRenderProxy::GetProviderData(bool&
 {
 	bOutValid = true;
 	return {};
-}
\ No newline at end of file
+}
diff --git a/Engine/Plugins/Experimental/DynamicWind/Source/DynamicWind/Private/DynamicWindProvider.cpp b/Engine/Plugins/Experimental/DynamicWind/Source/DynamicWind/Private/DynamicWindProvider.cpp
--- a/Engine/Plugins/Experimental/DynamicWind/Source/DynamicWind/Private/DynamicWindProvider.cpp
+++ b/Engine/Plugins/Experimental/DynamicWind/Source/DynamicWind/Private/DynamicWindProvider.cpp
@@ -309,7 +309,7 @@ void FDynamicWindTransformProvider::ProvideTransforms(FSkinningTransformProvider
 
 				FSetElementId SkeletonId = SkeletonLookup.FindId(SkinnedAnimationBatch.SkeletonGuid);
 
-				for (uint32 AnimIndex = 0; AnimIndex < UniqueAnimationCount; ++AnimIndex)
+				for (uint32 DirectionalityIndex = 0; DirectionalityIndex < UniqueAnimationCount; ++DirectionalityIndex)
 				{
 					uint32 BoneOffset = 0;
 
@@ -319,7 +319,7 @@ void FDynamicWindTransformProvider::ProvideTransforms(FSkinningTransformProvider
 						BlockHeaders[BlockWrite].BlockDstTransformOffset = DstTransformOffset;
 						BlockHeaders[BlockWrite].BlockNumBones = BonesPerGroup;
 						BlockHeaders[BlockWrite].BlockBoneOffset = BoneOffset;
-						BlockHeaders[BlockWrite].BlockAnimIndex = AnimIndex;
+						BlockHeaders[BlockWrite].BlockDirectionalityIndex = DirectionalityIndex;
 						BlockHeaders[BlockWrite].TotalTransformCount = TransformCount;
 						
 						BoneOffset += BonesPerGroup;
@@ -333,7 +333,7 @@ void FDynamicWindTransformProvider::ProvideTransforms(FSkinningTransformProvider
 						BlockHeaders[BlockWrite].BlockDstTransformOffset = DstTransformOffset;
 						BlockHeaders[BlockWrite].BlockNumBones = PartialTransformCount;
 						BlockHeaders[BlockWrite].BlockBoneOffset = BoneOffset;
-						BlockHeaders[BlockWrite].BlockAnimIndex = AnimIndex;
+						BlockHeaders[BlockWrite].BlockDirectionalityIndex = DirectionalityIndex;
 						BlockHeaders[BlockWrite].TotalTransformCount = TransformCount;
 
 						BoneOffset += PartialTransformCount;
diff --git a/Engine/Plugins/Experimental/DynamicWind/Source/DynamicWind/Public/DynamicWindData.h b/Engine/Plugins/Experimental/DynamicWind/Source/DynamicWind/Public/DynamicWindData.h
--- a/Engine/Plugins/Experimental/DynamicWind/Source/DynamicWind/Public/DynamicWindData.h
+++ b/Engine/Plugins/Experimental/DynamicWind/Source/DynamicWind/Public/DynamicWindData.h
@@ -18,6 +18,7 @@ public:
 	virtual const uint32 GetUniqueAnimationCount() const override;
 	virtual bool HasAnimationBounds() const override;
 	virtual bool GetAnimationBounds(uint32 AnimationIndex, FRenderBounds& OutBounds) const override;
+	virtual uint32 GetSkinningDataOffset(int32 InstanceIndex, const FTransform& ComponentTransform, const FSkinnedMeshInstanceData& InstanceData) const override;
 	virtual FTransformProviderRenderProxy* CreateRenderThreadResources(FSkinningSceneExtensionProxy* SceneProxy, FSceneInterface& Scene, FRHICommandListBase& RHICmdList) override;
 	virtual void DestroyRenderThreadResources(FTransformProviderRenderProxy* ProviderProxy) override;
 };
diff --git a/Engine/Source/Runtime/Engine/Classes/Animation/AnimBank.h b/Engine/Source/Runtime/Engine/Classes/Animation/AnimBank.h
--- a/Engine/Source/Runtime/Engine/Classes/Animation/AnimBank.h
+++ b/Engine/Source/Runtime/Engine/Classes/Animation/AnimBank.h
@@ -622,6 +622,7 @@ public:
 	virtual const uint32 GetUniqueAnimationCount() const override;
 	virtual bool HasAnimationBounds() const override;
 	virtual bool GetAnimationBounds(uint32 AnimationIndex, FRenderBounds& OutBounds) const override;
+	virtual uint32 GetSkinningDataOffset(int32 InstanceIndex, const FTransform& ComponentTransform, const FSkinnedMeshInstanceData& InstanceData) const override;
 	virtual FTransformProviderRenderProxy* CreateRenderThreadResources(FSkinningSceneExtensionProxy* SceneProxy, FSceneInterface& Scene, FRHICommandListBase& RHICmdList) override;
 	virtual void DestroyRenderThreadResources(FTransformProviderRenderProxy* ProviderProxy) override;
 	virtual bool IsCompiling() const override;
diff --git a/Engine/Source/Runtime/Engine/Private/Animation/AnimBank.cpp b/Engine/Source/Runtime/Engine/Private/Animation/AnimBank.cpp
--- a/Engine/Source/Runtime/Engine/Private/Animation/AnimBank.cpp
+++ b/Engine/Source/Runtime/Engine/Private/Animation/AnimBank.cpp
@@ -1277,6 +1277,17 @@ bool UAnimBankData::GetAnimationBounds(uint32 AnimationIndex, FRenderBounds& Out
 	return false;
 }
 
+uint32 UAnimBankData::GetSkinningDataOffset(int32 InstanceIndex, const FTransform& ComponentTransform, const FSkinnedMeshInstanceData& InstanceData) const
+{
+	const uint32 AnimationIndex = InstanceData.AnimationIndex;
+	if (ensure(AnimationIndex < uint32(AnimBankItems.Num())))
+	{
+		return AnimationIndex * 2u;
+	}
+	
+	return 0u;
+}
+
 FTransformProviderRenderProxy* UAnimBankData::CreateRenderThreadResources(FSkinningSceneExtensionProxy* SceneProxy, FSceneInterface& Scene, FRHICommandListBase& RHICmdList)
 {
 	FAnimBankDataRenderProxy* ProviderProxy = new FAnimBankDataRenderProxy(this, SceneProxy, Scene);
diff --git a/Engine/Source/Runtime/Engine/Private/Components/InstancedSkinnedMeshComponent.cpp b/Engine/Source/Runtime/Engine/Private/Components/InstancedSkinnedMeshComponent.cpp
--- a/Engine/Source/Runtime/Engine/Private/Components/InstancedSkinnedMeshComponent.cpp
+++ b/Engine/Source/Runtime/Engine/Private/Components/InstancedSkinnedMeshComponent.cpp
@@ -374,7 +374,7 @@ FInstanceDataManagerSourceDataDesc UInstancedSkinnedMeshComponent::GetComponentD
 	const bool bForceRefPose = CVarInstancedSkinnedMeshesForceRefPose.GetValueOnAnyThread() != 0;
 	const bool bValidTransformProvider = !bForceRefPose && TransformProvider != nullptr && TransformProvider->IsEnabled();
 
-	const uint32 UniqueAnimationCount = TransformProvider ? TransformProvider->GetUniqueAnimationCount() : 1;
+	const FTransform& ComponentTransform = GetComponentTransform();
 
 	Flags.bHasPerInstanceHierarchyOffset = false;
 	Flags.bHasPerInstanceLocalBounds = TransformProvider ? TransformProvider->HasAnimationBounds() : false;
@@ -402,17 +402,18 @@ FInstanceDataManagerSourceDataDesc UInstancedSkinnedMeshComponent::GetComponentD
 		ChangeSet.GetSkinningDataWriter().Gather(
 			[&](int32 InstanceIndex) -> uint32 
 			{
-				const uint32 AnimationIndex = InstanceData[InstanceIndex].AnimationIndex;
-				if (bValidTransformProvider && AnimationIndex < UniqueAnimationCount)
+				if (bValidTransformProvider)
 				{
-					return AnimationIndex * 2u;
+					check(InstanceData.IsValidIndex(InstanceIndex));
+					const FSkinnedMeshInstanceData& Instance = InstanceData[InstanceIndex];
+					return TransformProvider->GetSkinningDataOffset(InstanceIndex, ComponentTransform, Instance);
 				}
 
 				return 0u;
 			});
 
 		ChangeSet.GetLocalBoundsWriter().Gather(
-			[&](int32 InstanceIndex) -> FRenderBounds 
+			[&](int32 InstanceIndex) -> FRenderBounds
 			{
 				const uint32 AnimationIndex = InstanceData[InstanceIndex].AnimationIndex;
 				if (bValidTransformProvider && Flags.bHasPerInstanceLocalBounds)
@@ -437,8 +438,6 @@ FInstanceDataManagerSourceDataDesc UInstancedSkinnedMeshComponent::GetComponentD
 		ChangeSet.SetEditorData(HitProxies, SelectedInstances);
 	}
 #endif
-
-
 	};
 
 	return ComponentDesc;
diff --git a/Engine/Source/Runtime/Engine/Public/Animation/TransformProviderData.h b/Engine/Source/Runtime/Engine/Public/Animation/TransformProviderData.h
--- a/Engine/Source/Runtime/Engine/Public/Animation/TransformProviderData.h
+++ b/Engine/Source/Runtime/Engine/Public/Animation/TransformProviderData.h
@@ -9,6 +9,7 @@
 #include "TransformProviderData.generated.h"
 
 struct FRenderBounds;
+struct FSkinnedMeshInstanceData;
 class FSkinningSceneExtensionProxy;
 class FSceneInterface;
 class FRHICommandListBase;
@@ -69,6 +70,11 @@ public:
 		return false;
 	}
 
+	virtual uint32 GetSkinningDataOffset(int32 InstanceIndex, const FTransform& ComponentTransform, const FSkinnedMeshInstanceData& InstanceData) const
+	{
+		return 0u;
+	}
+
 	virtual FTransformProviderRenderProxy* CreateRenderThreadResources(FSkinningSceneExtensionProxy* SceneProxy, FSceneInterface& Scene, FRHICommandListBase& RHICmdList)
 	{
 		return nullptr;
diff --git a/Engine/Source/Runtime/Engine/Public/InstancedSkinnedMeshComponentHelper.h b/Engine/Source/Runtime/Engine/Public/InstancedSkinnedMeshComponentHelper.h
--- a/Engine/Source/Runtime/Engine/Public/InstancedSkinnedMeshComponentHelper.h
+++ b/Engine/Source/Runtime/Engine/Public/InstancedSkinnedMeshComponentHelper.h
@@ -52,8 +52,6 @@ FInstanceDataManagerSourceDataDesc FInstancedSkinnedMeshComponentHelper::GetComp
 	const bool bForceRefPose = UInstancedSkinnedMeshComponent::ShouldForceRefPose();
 	const bool bValidTransformProvider = !bForceRefPose && TransformProvider != nullptr && TransformProvider->IsEnabled();
 
-	const uint32 UniqueAnimationCount = TransformProvider ? TransformProvider->GetUniqueAnimationCount() : 1;
-
 	Flags.bHasPerInstanceHierarchyOffset = false;
 	Flags.bHasPerInstanceLocalBounds = TransformProvider ? TransformProvider->HasAnimationBounds() : false;
 	Flags.bHasPerInstanceDynamicData = false;
@@ -72,10 +70,9 @@ FInstanceDataManagerSourceDataDesc FInstancedSkinnedMeshComponentHelper::GetComp
 	ComponentDesc.PrimitiveLocalToWorld = InComponent.GetRenderMatrix();
 	ComponentDesc.ComponentMobility = InComponent.GetMobility();
 
-	const FReferenceSkeleton& RefSkeleton = SkinnedAsset->GetRefSkeleton();
-	uint32 MaxBoneTransformCount = RefSkeleton.GetRawBoneNum();
+	const FTransform& ComponentTransform = InComponent.GetComponentTransform();
 
-	ComponentDesc.BuildChangeSet = [&, MaxBoneTransformCount, MeshBounds = ComponentDesc.MeshBounds](FInstanceUpdateChangeSet& ChangeSet)
+	ComponentDesc.BuildChangeSet = [&, MeshBounds = ComponentDesc.MeshBounds](FInstanceUpdateChangeSet& ChangeSet)
 	{
 		// publish data
 		ChangeSet.GetTransformWriter().Gather([&](int32 InstanceIndex) -> FRenderTransform { return FRenderTransform(InComponent.InstanceData[InstanceIndex].Transform.ToMatrixWithScale()); });
@@ -84,10 +81,11 @@ FInstanceDataManagerSourceDataDesc FInstancedSkinnedMeshComponentHelper::GetComp
 		ChangeSet.GetSkinningDataWriter().Gather(
 			[&](int32 InstanceIndex) -> uint32
 			{
-				const uint32 AnimationIndex = InComponent.InstanceData[InstanceIndex].AnimationIndex;
-				if (bValidTransformProvider && AnimationIndex < UniqueAnimationCount)
+				if (bValidTransformProvider)
 				{
-					return AnimationIndex * MaxBoneTransformCount * 2u;
+					check(InComponent.InstanceData.IsValidIndex(InstanceIndex));
+					const FSkinnedMeshInstanceData& Instance = InComponent.InstanceData[InstanceIndex];
+					return TransformProvider->GetSkinningDataOffset(InstanceIndex, ComponentTransform, Instance);
 				}
 
 				return 0u;
