[DynamicWind] Fix a crash that can happen when using DynamicWind provider data on an Instanced Skinned Mesh component with a skeletal mesh that has no DynamicWind user data.

#rnx

diff --git a/Engine/Plugins/Experimental/DynamicWind/Source/DynamicWind/Private/DynamicWindData.cpp b/Engine/Plugins/Experimental/DynamicWind/Source/DynamicWind/Private/DynamicWindData.cpp
--- a/Engine/Plugins/Experimental/DynamicWind/Source/DynamicWind/Private/DynamicWindData.cpp
+++ b/Engine/Plugins/Experimental/DynamicWind/Source/DynamicWind/Private/DynamicWindData.cpp
@@ -77,7 +77,6 @@ FDynamicWindDataRenderProxy::FDynamicWindDataRenderProxy(UDynamicWindData* BankD
 	
 	const USkeletalMesh* SkeletalMesh = CastChecked<USkeletalMesh>(SkinnedAsset);
 	SkeletalData = const_cast<USkeletalMesh*>(SkeletalMesh)->GetAssetUserData<UDynamicWindSkeletalData>();
-	check(SkeletalData);
 
 	// GW-TODO: Need a validation method on the base data type to check per-proxy if ref pose fallback is needed.
 }
diff --git a/Engine/Plugins/Experimental/DynamicWind/Source/DynamicWind/Private/DynamicWindProvider.cpp b/Engine/Plugins/Experimental/DynamicWind/Source/DynamicWind/Private/DynamicWindProvider.cpp
--- a/Engine/Plugins/Experimental/DynamicWind/Source/DynamicWind/Private/DynamicWindProvider.cpp
+++ b/Engine/Plugins/Experimental/DynamicWind/Source/DynamicWind/Private/DynamicWindProvider.cpp
@@ -142,15 +142,22 @@ FDynamicWindTransformProvider::~FDynamicWindTransformProvider()
 	
 void FDynamicWindTransformProvider::RegisterSceneProxy(const Nanite::FSkinnedSceneProxy* InProxy)
 {
-	USkeletalMesh* SkeletalMesh = const_cast<USkeletalMesh*>(Cast<USkeletalMesh>(InProxy->GetSkinnedAsset()));
+	USkeletalMesh* SkeletalMesh = const_cast<USkeletalMesh*>(CastChecked<USkeletalMesh>(InProxy->GetSkinnedAsset()));
 	const USkeleton* Skeleton = SkeletalMesh->GetSkeleton();
 	const FReferenceSkeleton& RefSkeleton = Skeleton->GetReferenceSkeleton();
 	const uint32 NumBonesInSkeleton = RefSkeleton.GetNum();
+	
 	const UDynamicWindSkeletalData* Data = SkeletalMesh->GetAssetUserData<UDynamicWindSkeletalData>();
+	if (Data == nullptr)
+	{
+		// No user data, don't register the skeleton. The wind simulation CS will provide identity
+		return;
+	}
+
 	const uint64 UserDataHash = Data->GetSkeletalDataHash();
 	
 	FSkeletonEntry& SkeletonEntry = SkeletonLookup.FindOrAdd(Skeleton->GetGuid());
-	
+
 	++SkeletonEntry.ReferenceCount;
 	if (SkeletonEntry.ReferenceCount == 1)
 	{
