Remove the explicit UnregisterProvider  in the FDynamicWindTransformProvider  destructor, as FScene should have already removed it when it gets destroyed.

diff --git a/Engine/Plugins/Experimental/DynamicWind/Source/DynamicWind/Private/DynamicWindProvider.cpp b/Engine/Plugins/Experimental/DynamicWind/Source/DynamicWind/Private/DynamicWindProvider.cpp
--- a/Engine/Plugins/Experimental/DynamicWind/Source/DynamicWind/Private/DynamicWindProvider.cpp
+++ b/Engine/Plugins/Experimental/DynamicWind/Source/DynamicWind/Private/DynamicWindProvider.cpp
@@ -133,12 +133,7 @@ FDynamicWindTransformProvider::FDynamicWindTransformProvider(FScene& InScene)
 }
 
 FDynamicWindTransformProvider::~FDynamicWindTransformProvider()
-{
-	if (auto TransformProvider = Scene.GetExtensionPtr<FSkinningTransformProvider>())
-	{
-		TransformProvider->UnregisterProvider(GDynamicWindTransformProviderId);
-	}
-}
+{}
 	
 void FDynamicWindTransformProvider::RegisterSceneProxy(const Nanite::FSkinnedSceneProxy* InProxy)
 {
