[DynamicWind] Fix some thread safety issues with the DynamicWind subsystem's GT/RT interop, one of which was causing crashes in the transform provider when shutting down a level. Moving all render thread operation into the same interface boundary (FDynamicWindProvider) and only interacting with it on the render thread after creation fixes these crashes and data races, and makes the separation of responsibility a little more clear.

diff --git a/Engine/Plugins/Experimental/DynamicWind/Source/DynamicWind/Private/DynamicWindProvider.cpp b/Engine/Plugins/Experimental/DynamicWind/Source/DynamicWind/Private/DynamicWindProvider.cpp
--- a/Engine/Plugins/Experimental/DynamicWind/Source/DynamicWind/Private/DynamicWindProvider.cpp
+++ b/Engine/Plugins/Experimental/DynamicWind/Source/DynamicWind/Private/DynamicWindProvider.cpp
@@ -2,6 +2,7 @@
 
 #include "DynamicWindProvider.h"
 #include "DynamicWindSkeletalData.h"
+#include "DynamicWindLog.h"
 #include "Animation/Skeleton.h"
 #include "Engine/Texture.h"
 #include "GlobalRenderResources.h"
@@ -23,6 +24,8 @@ DECLARE_GPU_STAT(DynamicWind);
 namespace
 {
 
+static const FGuid DynamicWindTransformProviderId(DYNAMIC_WIND_TRANSFORM_PROVIDER_GUID);
+
 TAutoConsoleVariable<float> CVarDynamicWindOverrideSpeed(
 	TEXT("DynamicWind.OverrideSpeed"),
 	-1.0f,
@@ -120,7 +123,17 @@ IMPLEMENT_GLOBAL_SHADER(FDynamicWindEvalCS, "/Plugin/DynamicWind/DynamicWindEval
 FDynamicWindTransformProvider::FDynamicWindTransformProvider(FScene& InScene)
 : Scene(InScene)
 , BoneDataBuffer(1, TEXT("DynamicWind.BoneDataBuffer"))
+{}
+
+FDynamicWindTransformProvider::~FDynamicWindTransformProvider()
+{	
+	check(!bIsRegistered);
+}
+
+bool FDynamicWindTransformProvider::Register()
 {
+	check(!bIsRegistered);
+
 	if (auto TransformProvider = Scene.GetExtensionPtr<FSkinningTransformProvider>())
 	{
 		// Register GPU dynamic wind transform provider
@@ -129,15 +142,88 @@ FDynamicWindTransformProvider::FDynamicWindTransformProvider(FScene& InScene)
 			FSkinningTransformProvider::FOnProvideTransforms::CreateRaw(this, &FDynamicWindTransformProvider::ProvideTransforms),
 			true /* Use skeleton batching */
 		);
+
+		OnProxyCreateRenderThreadResources =
+			Nanite::FSkinnedSceneProxyDelegates::OnCreateRenderThreadResources.AddLambda(
+				[this](const Nanite::FSkinnedSceneProxy* InSkinnedSceneProxy)
+				{
+					if (&InSkinnedSceneProxy->GetScene() == &Scene)
+					{
+						PRAGMA_DISABLE_DEPRECATION_WARNINGS
+						RegisterSceneProxy(InSkinnedSceneProxy);
+						PRAGMA_ENABLE_DEPRECATION_WARNINGS
+					}
+				}
+			);
+
+		OnProxyDestroyRenderThreadResources =
+			Nanite::FSkinnedSceneProxyDelegates::OnDestroyRenderThreadResources.AddLambda(
+				[this](const Nanite::FSkinnedSceneProxy* InSkinnedSceneProxy)
+				{
+					if (&InSkinnedSceneProxy->GetScene() == &Scene)
+					{
+						PRAGMA_DISABLE_DEPRECATION_WARNINGS
+						UnregisterSceneProxy(InSkinnedSceneProxy);
+						PRAGMA_ENABLE_DEPRECATION_WARNINGS
+					}
+				}
+			);
+
+		bIsRegistered = true;
+		return true;
 	}
+
+	return false;
 }
 
-FDynamicWindTransformProvider::~FDynamicWindTransformProvider()
-{}
+bool FDynamicWindTransformProvider::Unregister()
+{
+	if (bIsRegistered)
+	{
+		if (auto TransformProvider = Scene.GetExtensionPtr<FSkinningTransformProvider>())
+		{
+			TransformProvider->UnregisterProvider(DynamicWindTransformProviderId);
+		}
+
+		if (OnProxyCreateRenderThreadResources.IsValid())
+		{
+			Nanite::FSkinnedSceneProxyDelegates::OnCreateRenderThreadResources.Remove(OnProxyCreateRenderThreadResources);
+			OnProxyCreateRenderThreadResources.Reset();
+		}
+
+		if (OnProxyDestroyRenderThreadResources.IsValid())
+		{
+			Nanite::FSkinnedSceneProxyDelegates::OnDestroyRenderThreadResources.Remove(OnProxyDestroyRenderThreadResources);
+			OnProxyDestroyRenderThreadResources.Reset();
+		}
+
+		bIsRegistered = false;
+		return true;
+	}
+
+	return false;
+}
 	
 void FDynamicWindTransformProvider::RegisterSceneProxy(const Nanite::FSkinnedSceneProxy* InProxy)
 {
-	USkeletalMesh* SkeletalMesh = const_cast<USkeletalMesh*>(CastChecked<USkeletalMesh>(InProxy->GetSkinnedAsset()));
+	if (!bIsRegistered)
+	{
+		UE_LOG(LogDynamicWind, Verbose, TEXT("Trying to register SkinnedSceneProxy (SkinnedAsset=%s), when the provider is not registered."), *GetNameSafe(InProxy->GetSkinnedAsset()));
+		return;
+	}
+
+	const FSkinningSceneExtensionProxy* SkinningSceneExtProxy = InProxy->GetSkinningSceneExtensionProxy();
+	if (SkinningSceneExtProxy == nullptr || SkinningSceneExtProxy->GetTransformProviderId() != DynamicWindTransformProviderId)
+	{
+		return;
+	}
+
+	USkeletalMesh* SkeletalMesh = const_cast<USkeletalMesh*>(Cast<USkeletalMesh>(InProxy->GetSkinnedAsset()));
+	if (SkeletalMesh == nullptr)
+	{
+		return;
+	}
+
 	const USkeleton* Skeleton = SkeletalMesh->GetSkeleton();
 	const FReferenceSkeleton& RefSkeleton = Skeleton->GetReferenceSkeleton();
 	const uint32 NumBonesInSkeleton = RefSkeleton.GetNum();
@@ -215,6 +301,18 @@ void FDynamicWindTransformProvider::RegisterSceneProxy(const Nanite::FSkinnedSce
 
 void FDynamicWindTransformProvider::UnregisterSceneProxy(const Nanite::FSkinnedSceneProxy* InProxy)
 {
+	if (!bIsRegistered)
+	{
+		UE_LOG(LogDynamicWind, Verbose, TEXT("Trying to unregister SkinnedSceneProxy (SkinnedAsset=%s), when the provider is not registered."), *GetNameSafe(InProxy->GetSkinnedAsset()));
+		return;
+	}
+
+	const FSkinningSceneExtensionProxy* SkinningSceneExtProxy = InProxy->GetSkinningSceneExtensionProxy();
+	if (SkinningSceneExtProxy == nullptr || SkinningSceneExtProxy->GetTransformProviderId() != DynamicWindTransformProviderId)
+	{
+		return;
+	}
+
 	const USkeletalMesh* SkeletalMesh = Cast<USkeletalMesh>(InProxy->GetSkinnedAsset());
 	const USkeleton* Skeleton = SkeletalMesh->GetSkeleton();
 	auto EntryID = SkeletonLookup.FindId(Skeleton->GetGuid());
diff --git a/Engine/Plugins/Experimental/DynamicWind/Source/DynamicWind/Private/DynamicWindSubsystem.cpp b/Engine/Plugins/Experimental/DynamicWind/Source/DynamicWind/Private/DynamicWindSubsystem.cpp
--- a/Engine/Plugins/Experimental/DynamicWind/Source/DynamicWind/Private/DynamicWindSubsystem.cpp
+++ b/Engine/Plugins/Experimental/DynamicWind/Source/DynamicWind/Private/DynamicWindSubsystem.cpp
@@ -17,6 +17,7 @@
 #include "ScenePrivate.h"
 #include "NaniteSceneProxy.h"
 #include "DynamicWindSkeletalData.h"
+#include "DynamicWindProvider.h"
 
 namespace
 {
@@ -30,13 +31,6 @@ TAutoConsoleVariable<bool> CVarDynamicWindEnable(
 
 }
 
-UDynamicWindSubsystem::UDynamicWindSubsystem()
-: bInitialized(false)
-, TransformProvider(nullptr)
-, TransformProviderId(DYNAMIC_WIND_TRANSFORM_PROVIDER_GUID)
-{
-}
-
 bool UDynamicWindSubsystem::ShouldCreateSubsystem(UObject* Outer) const
 {
 	return CVarDynamicWindEnable.GetValueOnAnyThread();
@@ -52,111 +46,48 @@ void UDynamicWindSubsystem::PostInitialize()
 		{
 			if (FScene* RenderScene = Scene->GetRenderScene())
 			{
-				TransformProvider = MakeUnique<FDynamicWindTransformProvider>(*RenderScene);
-
-				OnCreateRenderThreadResources =
-					Nanite::FSkinnedSceneProxyDelegates::OnCreateRenderThreadResources.AddLambda(
-						[this](const Nanite::FSkinnedSceneProxy* InSkinnedSceneProxy)
-						{
-							RegisterSceneProxy(InSkinnedSceneProxy);
-						}
-					);
-
-				OnDestroyRenderThreadResources =
-					Nanite::FSkinnedSceneProxyDelegates::OnDestroyRenderThreadResources.AddLambda(
-						[this](const Nanite::FSkinnedSceneProxy* InSkinnedSceneProxy)
-						{
-							UnregisterSceneProxy(InSkinnedSceneProxy);
-						}
-					);
-
-
+				TransformProvider = MakeShared<FDynamicWindTransformProvider, ESPMode::ThreadSafe>(*RenderScene);
+
+				// Transform provider must be registered on the render thread
+				ENQUEUE_RENDER_COMMAND(RegisterDynamicWindProvider)
+				(
+					[ProviderToRegister = TransformProvider](FRHICommandListImmediate&)
+					{
+						ProviderToRegister->Register();
+					}
+				);
 			}
 		}
 	}
-
-	bInitialized =
-		TransformProvider.IsValid() &&
-		OnCreateRenderThreadResources.IsValid() &&
-		OnDestroyRenderThreadResources.IsValid();
 }
 
 void UDynamicWindSubsystem::PreDeinitialize()
 {
-	bInitialized = false;
-
-	if (OnCreateRenderThreadResources.IsValid())
-	{
-		Nanite::FSkinnedSceneProxyDelegates::OnCreateRenderThreadResources.Remove(OnCreateRenderThreadResources);
-		OnCreateRenderThreadResources.Reset();
-	}
-
-	if (OnDestroyRenderThreadResources.IsValid())
+	if (TransformProvider.IsValid())
 	{
-		Nanite::FSkinnedSceneProxyDelegates::OnDestroyRenderThreadResources.Remove(OnDestroyRenderThreadResources);
-		OnDestroyRenderThreadResources.Reset();
+		// Transform provider must be unregistered and destroyed on the render thread
+		ENQUEUE_RENDER_COMMAND(UnregisterAndDestroyDynamicWindProvider)
+		(
+			[ProviderToDestroy = MoveTemp(TransformProvider)](FRHICommandListImmediate&) mutable
+			{
+				ProviderToDestroy->Unregister();
+				ProviderToDestroy.Reset(); // should be the last reference out
+			}
+		);
 	}
 
-	TransformProvider.Reset();
-
 	Super::PreDeinitialize();
 }
 
 void UDynamicWindSubsystem::RegisterSceneProxy(const Nanite::FSkinnedSceneProxy* InSkinnedProxy)
-{
-	if (&InSkinnedProxy->GetScene() != GetWorld()->Scene)
-	{
-		return;
-	}
-
-	FSkinningSceneExtensionProxy* Proxy = InSkinnedProxy->GetSkinningSceneExtensionProxy();
-
-	if (!Proxy || Proxy->GetTransformProviderId() != TransformProviderId)
-	{
-		return;
-	}
-
-	if (const USkinnedAsset* SkinnedAsset = Proxy->GetSkinnedAsset())
-	{
-		if (!bInitialized)
-		{
-			UE_LOG(LogDynamicWind, Verbose, TEXT("Trying to register SkinnedSceneProxy (SkinnedAsset=%s), when the subsystem is not initialized."), *SkinnedAsset->GetName());
-			return;
-		}
-
-		TransformProvider->RegisterSceneProxy(InSkinnedProxy);
-	}
-}
+{}
 
 void UDynamicWindSubsystem::UnregisterSceneProxy(const Nanite::FSkinnedSceneProxy* InSkinnedProxy)
-{
-	if (&InSkinnedProxy->GetScene() != GetWorld()->Scene)
-	{
-		return;
-	}
-
-	FSkinningSceneExtensionProxy* Proxy = InSkinnedProxy->GetSkinningSceneExtensionProxy();
-
-	if (!Proxy || Proxy->GetTransformProviderId() != TransformProviderId)
-	{
-		return;
-	}
-
-	if (const USkinnedAsset* SkinnedAsset = Proxy->GetSkinnedAsset())
-	{
-		if (!bInitialized)
-		{
-			UE_LOG(LogDynamicWind, Verbose, TEXT("Trying to unregister SkinnedSceneProxy (SkinnedAsset=%s), when the subsystem is not initialized."), *SkinnedAsset->GetName());
-			return;
-		}
-
-		TransformProvider->UnregisterSceneProxy(InSkinnedProxy);
-	}
-}
+{}
 
 float UDynamicWindSubsystem::GetBlendedWindAmplitude() const
 {
-	if (TransformProvider)
+	if (TransformProvider.IsValid())
 	{
 		return TransformProvider->GetBlendedWindAmplitude();
 	}
@@ -166,11 +97,11 @@ float UDynamicWindSubsystem::GetBlendedWindAmplitude() const
 
 void UDynamicWindSubsystem::UpdateWindParameters(const FDynamicWindParameters& Parameters)
 {
-	if (bInitialized)
+	if (TransformProvider.IsValid())
 	{
-		ENQUEUE_RENDER_COMMAND(UpdateFoliageWindSimulationInputs)
+		ENQUEUE_RENDER_COMMAND(UpdateDynamicWindParameters)
 		(
-			[WindProvider = TransformProvider.Get(), WindParameters = Parameters](FRHICommandListImmediate&) mutable
+			[WindProvider=TransformProvider, WindParameters=Parameters](FRHICommandListImmediate&) mutable
 			{
 				WindProvider->UpdateParameters(WindParameters);
 			}
diff --git a/Engine/Plugins/Experimental/DynamicWind/Source/DynamicWind/Public/DynamicWindProvider.h b/Engine/Plugins/Experimental/DynamicWind/Source/DynamicWind/Public/DynamicWindProvider.h
--- a/Engine/Plugins/Experimental/DynamicWind/Source/DynamicWind/Public/DynamicWindProvider.h
+++ b/Engine/Plugins/Experimental/DynamicWind/Source/DynamicWind/Public/DynamicWindProvider.h
@@ -23,8 +23,13 @@ public:
 	explicit FDynamicWindTransformProvider(FScene& InScene);
 	FDynamicWindTransformProvider() = delete;
 	~FDynamicWindTransformProvider();
-	
+
+	bool Register();
+	bool Unregister();
+
+	UE_DEPRECATED(5.7, "This method will be made private in the next release. It shouldn't be used externally")
 	void RegisterSceneProxy(const Nanite::FSkinnedSceneProxy* InProxy);
+	UE_DEPRECATED(5.7, "This method will be made private in the next release. It shouldn't be used externally")
 	void UnregisterSceneProxy(const Nanite::FSkinnedSceneProxy* InProxy);
 
 	int32 UpdateParameters(const FDynamicWindParameters& Parameters);
@@ -62,4 +67,9 @@ private:
 	uint64 LastSimulatedFrameNumber = 0u;
 
 	float BlendedWindAmplitude = -1.0f;
+
+	FDelegateHandle OnProxyCreateRenderThreadResources;
+	FDelegateHandle OnProxyDestroyRenderThreadResources;
+
+	bool bIsRegistered = false;
 };
diff --git a/Engine/Plugins/Experimental/DynamicWind/Source/DynamicWind/Public/DynamicWindSubsystem.h b/Engine/Plugins/Experimental/DynamicWind/Source/DynamicWind/Public/DynamicWindSubsystem.h
--- a/Engine/Plugins/Experimental/DynamicWind/Source/DynamicWind/Public/DynamicWindSubsystem.h
+++ b/Engine/Plugins/Experimental/DynamicWind/Source/DynamicWind/Public/DynamicWindSubsystem.h
@@ -8,11 +8,12 @@
 #include "Containers/Map.h"
 #include "DynamicWindLog.h"
 #include "DynamicWindParameters.h"
-#include "DynamicWindProvider.h"
 #include "DynamicWindSubsystem.generated.h"
 
 class USkeleton;
 class UInstancedSkinnedMeshComponent;
+class FDynamicWindTransformProvider;
+namespace Nanite { class FSkinnedSceneProxy; }
 
 UCLASS(BlueprintType, ClassGroup = (Rendering, Common))
 class DYNAMICWIND_API UDynamicWindSubsystem : public UWorldSubsystem
@@ -20,14 +21,14 @@ class DYNAMICWIND_API UDynamicWindSubsystem : public UWorldSubsystem
 	GENERATED_BODY()
 
 public:
-	UDynamicWindSubsystem();
-
 	virtual bool ShouldCreateSubsystem(UObject* Outer) const override;
 
 	virtual void PostInitialize() override;
 	virtual void PreDeinitialize() override;
 
+	UE_DEPRECATED(5.7, "This method no longer has any effect. It shouldn't be used externally and will be removed next release")
 	void RegisterSceneProxy(const Nanite::FSkinnedSceneProxy* InSkinnedProxy);
+	UE_DEPRECATED(5.7, "This method no longer has any effect. It shouldn't be used externally and will be removed next release")
 	void UnregisterSceneProxy(const Nanite::FSkinnedSceneProxy* InSkinnedProxy);
 
 public:
@@ -38,11 +39,7 @@ public:
 	void UpdateWindParameters(const FDynamicWindParameters& Parameters);
 
 private:
-	bool bInitialized = false;
-
-	TUniquePtr<FDynamicWindTransformProvider> TransformProvider;
-	FGuid TransformProviderId;
+	using FTransformProviderPtr = TSharedPtr<FDynamicWindTransformProvider, ESPMode::ThreadSafe>;
 
-	FDelegateHandle OnCreateRenderThreadResources;
-	FDelegateHandle OnDestroyRenderThreadResources;
+	FTransformProviderPtr TransformProvider;
 };
