[DynamicWind] Move DynamicWindProvider to Private as there's currently no good reason to expose it or its dependencies. Also, this fixes compiler issues with external users of this library attempting to use the DynamicWindSubsystem.

#Synced-CL 48677745

diff --git a/Engine/Plugins/Experimental/DynamicWind/Source/DynamicWind/Private/DynamicWindProvider.h b/Engine/Plugins/Experimental/DynamicWind/Source/DynamicWind/Private/DynamicWindProvider.h
new file mode 100644
--- /dev/null
+++ b/Engine/Plugins/Experimental/DynamicWind/Source/DynamicWind/Private/DynamicWindProvider.h
@@ -0,0 +1,65 @@
+// Copyright Epic Games, Inc. All Rights Reserved.
+
+#pragma once
+
+#include "Engine/AssetUserData.h"
+#include "Engine/SkeletalMesh.h"
+#include "HLSLTypeAliases.h"
+#include "Shared/DynamicWindCommon.ush"
+#include "DynamicWindParameters.h"
+#include "Skinning/SkinningTransformProvider.h"
+#include "SkinningDefinitions.h"
+#include "Animation/TransformProviderData.h"
+
+class USkeleton;
+namespace Nanite { class FSkinnedSceneProxy; }
+
+#define DYNAMIC_WIND_TRANSFORM_PROVIDER_GUID 0xDFD4874B, 0xEE57466D, 0x883D6419, 0xAFE99EAC
+#define DYNAMIC_WIND_DEBUG_SKELETON_NAMES (UE_BUILD_DEBUG | UE_BUILD_DEVELOPMENT)
+
+class FDynamicWindTransformProvider
+{
+public:
+	explicit FDynamicWindTransformProvider(FScene& InScene);
+	FDynamicWindTransformProvider() = delete;
+	~FDynamicWindTransformProvider();
+	
+	void RegisterSceneProxy(const Nanite::FSkinnedSceneProxy* InProxy);
+	void UnregisterSceneProxy(const Nanite::FSkinnedSceneProxy* InProxy);
+
+	int32 UpdateParameters(const FDynamicWindParameters& Parameters);
+
+	float GetBlendedWindAmplitude() const;
+
+private:
+	using FBoneDataBuffer = TPersistentByteAddressBuffer<FDynamicWindBoneData>;
+	using FBoneDataUploader = TByteAddressBufferScatterUploader<FDynamicWindBoneData>;
+
+	struct FSkeletonEntry
+	{	
+		uint64 UserDataHash = 0;
+		uint32 NumBones = 0;
+		uint32 ReferenceCount = 0;
+	#if DYNAMIC_WIND_DEBUG_SKELETON_NAMES
+		FName SkeletonName;
+	#endif
+		
+		FDynamicWindSkeletonData Data;
+	};
+
+	void ProvideTransforms(FSkinningTransformProvider::FProviderContext& Context);
+
+private:
+	FScene& Scene;
+
+	TMap<FGuid, FSkeletonEntry> SkeletonLookup;
+	FSpanAllocator BoneDataAllocator;
+	TUniquePtr<FBoneDataUploader> BoneDataUploader;
+	FBoneDataBuffer BoneDataBuffer;
+
+	FDynamicWindParameters WindParameters = {};
+
+	uint64 LastSimulatedFrameNumber = 0u;
+
+	float BlendedWindAmplitude = -1.0f;
+};
diff --git a/Engine/Plugins/Experimental/DynamicWind/Source/DynamicWind/Private/DynamicWindSubsystem.cpp b/Engine/Plugins/Experimental/DynamicWind/Source/DynamicWind/Private/DynamicWindSubsystem.cpp
--- a/Engine/Plugins/Experimental/DynamicWind/Source/DynamicWind/Private/DynamicWindSubsystem.cpp
+++ b/Engine/Plugins/Experimental/DynamicWind/Source/DynamicWind/Private/DynamicWindSubsystem.cpp
@@ -17,6 +17,7 @@
 #include "ScenePrivate.h"
 #include "NaniteSceneProxy.h"
 #include "DynamicWindSkeletalData.h"
+#include "DynamicWindProvider.h"
 
 namespace
 {
diff --git a/Engine/Plugins/Experimental/DynamicWind/Source/DynamicWind/Public/DynamicWindProvider.h b/Engine/Plugins/Experimental/DynamicWind/Source/DynamicWind/Public/DynamicWindProvider.h
deleted file mode 100644
diff --git a/Engine/Plugins/Experimental/DynamicWind/Source/DynamicWind/Public/DynamicWindSubsystem.h b/Engine/Plugins/Experimental/DynamicWind/Source/DynamicWind/Public/DynamicWindSubsystem.h
--- a/Engine/Plugins/Experimental/DynamicWind/Source/DynamicWind/Public/DynamicWindSubsystem.h
+++ b/Engine/Plugins/Experimental/DynamicWind/Source/DynamicWind/Public/DynamicWindSubsystem.h
@@ -8,11 +8,12 @@
 #include "Containers/Map.h"
 #include "DynamicWindLog.h"
 #include "DynamicWindParameters.h"
-#include "DynamicWindProvider.h"
 #include "DynamicWindSubsystem.generated.h"
 
 class USkeleton;
 class UInstancedSkinnedMeshComponent;
+class FDynamicWindTransformProvider;
+namespace Nanite { class FSkinnedSceneProxy; }
 
 UCLASS(BlueprintType, ClassGroup = (Rendering, Common))
 class DYNAMICWIND_API UDynamicWindSubsystem : public UWorldSubsystem
diff --git a/Engine/Plugins/Experimental/DynamicWind/Source/DynamicWindEditor/Private/DynamicWindFactory.cpp b/Engine/Plugins/Experimental/DynamicWind/Source/DynamicWindEditor/Private/DynamicWindFactory.cpp
--- a/Engine/Plugins/Experimental/DynamicWind/Source/DynamicWindEditor/Private/DynamicWindFactory.cpp
+++ b/Engine/Plugins/Experimental/DynamicWind/Source/DynamicWindEditor/Private/DynamicWindFactory.cpp
@@ -1,7 +1,6 @@
 // Copyright Epic Games, Inc. All Rights Reserved.
 
 #include "DynamicWindFactory.h"
-#include "DynamicWindProvider.h"
 #include "DynamicWindData.h"
 
 UDynamicWindDataFactory::UDynamicWindDataFactory(const FObjectInitializer& ObjectInitializer)
