[DynamicWind] Fix non-Unity compile errors seemingly stemmed from using a TUniquePtr of a pre-declared type in a UCLASS. Switching to TPimplPtr seems to fix it.

#Synced-CL 48677745

diff --git a/Engine/Plugins/Experimental/DynamicWind/Source/DynamicWind/Private/DynamicWindSubsystem.cpp b/Engine/Plugins/Experimental/DynamicWind/Source/DynamicWind/Private/DynamicWindSubsystem.cpp
--- a/Engine/Plugins/Experimental/DynamicWind/Source/DynamicWind/Private/DynamicWindSubsystem.cpp
+++ b/Engine/Plugins/Experimental/DynamicWind/Source/DynamicWind/Private/DynamicWindSubsystem.cpp
@@ -53,7 +53,7 @@ void UDynamicWindSubsystem::PostInitialize()
 		{
 			if (FScene* RenderScene = Scene->GetRenderScene())
 			{
-				TransformProvider = MakeUnique<FDynamicWindTransformProvider>(*RenderScene);
+				TransformProvider = MakePimpl<FDynamicWindTransformProvider>(*RenderScene);
 
 				OnCreateRenderThreadResources =
 					Nanite::FSkinnedSceneProxyDelegates::OnCreateRenderThreadResources.AddLambda(
diff --git a/Engine/Plugins/Experimental/DynamicWind/Source/DynamicWind/Public/DynamicWindSubsystem.h b/Engine/Plugins/Experimental/DynamicWind/Source/DynamicWind/Public/DynamicWindSubsystem.h
--- a/Engine/Plugins/Experimental/DynamicWind/Source/DynamicWind/Public/DynamicWindSubsystem.h
+++ b/Engine/Plugins/Experimental/DynamicWind/Source/DynamicWind/Public/DynamicWindSubsystem.h
@@ -41,7 +41,7 @@ public:
 private:
 	bool bInitialized = false;
 
-	TUniquePtr<FDynamicWindTransformProvider> TransformProvider;
+	TPimplPtr<FDynamicWindTransformProvider> TransformProvider;
 	FGuid TransformProviderId;
 
 	FDelegateHandle OnCreateRenderThreadResources;
