[DynamicWind] Fix for crash in dynamic wind subsystem when changing preview platforms during PIE and then stopping PIE.

#Synced-CL 49290026

diff --git a/Engine/Plugins/Experimental/DynamicWind/Source/DynamicWind/Private/DynamicWindSubsystem.cpp b/Engine/Plugins/Experimental/DynamicWind/Source/DynamicWind/Private/DynamicWindSubsystem.cpp
--- a/Engine/Plugins/Experimental/DynamicWind/Source/DynamicWind/Private/DynamicWindSubsystem.cpp
+++ b/Engine/Plugins/Experimental/DynamicWind/Source/DynamicWind/Private/DynamicWindSubsystem.cpp
@@ -42,39 +42,47 @@ void UDynamicWindSubsystem::PostInitialize()
 
 	if (UWorld* World = GetWorld())
 	{
-		if (FSceneInterface* Scene = World->Scene)
+		CreateTransformProvider(*World);
+	}
+
+#if WITH_EDITOR
+	OnPreRecreateScene = FWorldDelegates::OnPreRecreateScene.AddLambda(
+		[this](UWorld* World)
 		{
-			if (FScene* RenderScene = Scene->GetRenderScene())
+			if (World == GetWorld())
 			{
-				TransformProvider = MakeShared<FDynamicWindTransformProvider, ESPMode::ThreadSafe>(*RenderScene);
-
-				// Transform provider must be registered on the render thread
-				ENQUEUE_RENDER_COMMAND(RegisterDynamicWindProvider)
-				(
-					[ProviderToRegister = TransformProvider](FRHICommandListImmediate&)
-					{
-						ProviderToRegister->Register();
-					}
-				);
+				ReleaseTransformProvider();
 			}
 		}
-	}
+	);
+	OnPostRecreateScene = FWorldDelegates::OnPostRecreateScene.AddLambda(
+		[this](UWorld* World)
+		{
+			if (World == GetWorld())
+			{
+				CreateTransformProvider(*World);
+			}
+		}
+	);
+#endif // WITH_EDITOR
 }
 
 void UDynamicWindSubsystem::PreDeinitialize()
 {
-	if (TransformProvider.IsValid())
+	ReleaseTransformProvider();
+
+#if WITH_EDITOR
+	if (OnPreRecreateScene.IsValid())
 	{
-		// Transform provider must be unregistered and destroyed on the render thread
-		ENQUEUE_RENDER_COMMAND(UnregisterAndDestroyDynamicWindProvider)
-		(
-			[ProviderToDestroy = MoveTemp(TransformProvider)](FRHICommandListImmediate&) mutable
-			{
-				ProviderToDestroy->Unregister();
-				ProviderToDestroy.Reset(); // should be the last reference out
-			}
-		);
+		FWorldDelegates::OnPreRecreateScene.Remove(OnPreRecreateScene);
+		OnPreRecreateScene.Reset();
 	}
+	if (OnPostRecreateScene.IsValid())
+	{
+		FWorldDelegates::OnPostRecreateScene.Remove(OnPostRecreateScene);
+		OnPostRecreateScene.Reset();
+	}
+#endif // WITH_EDITOR
 
 	Super::PreDeinitialize();
 }
@@ -102,3 +110,39 @@ void UDynamicWindSubsystem::UpdateWindParameters(const FDynamicWindParameters& P
 		);
 	}
 }
+
+void UDynamicWindSubsystem::CreateTransformProvider(UWorld& World)
+{
+	if (FSceneInterface* Scene = World.Scene)
+	{
+		if (FScene* RenderScene = Scene->GetRenderScene())
+		{
+			TransformProvider = MakeShared<FDynamicWindTransformProvider, ESPMode::ThreadSafe>(*RenderScene);
+
+			// Transform provider must be registered on the render thread
+			ENQUEUE_RENDER_COMMAND(RegisterDynamicWindProvider)
+			(
+				[ProviderToRegister = TransformProvider](FRHICommandListImmediate&)
+				{
+					ProviderToRegister->Register();
+				}
+			);
+		}
+	}
+}
+
+void UDynamicWindSubsystem::ReleaseTransformProvider()
+{
+	if (TransformProvider.IsValid())
+	{
+		// Transform provider must be unregistered and destroyed on the render thread
+		ENQUEUE_RENDER_COMMAND(UnregisterAndDestroyDynamicWindProvider)
+		(
+			[ProviderToDestroy = MoveTemp(TransformProvider)](FRHICommandListImmediate&) mutable
+			{
+				ProviderToDestroy->Unregister();
+				ProviderToDestroy.Reset(); // should be the last reference out
+			}
+		);
+	}
+}
\ No newline at end of file
diff --git a/Engine/Plugins/Experimental/DynamicWind/Source/DynamicWind/Public/DynamicWindSubsystem.h b/Engine/Plugins/Experimental/DynamicWind/Source/DynamicWind/Public/DynamicWindSubsystem.h
--- a/Engine/Plugins/Experimental/DynamicWind/Source/DynamicWind/Public/DynamicWindSubsystem.h
+++ b/Engine/Plugins/Experimental/DynamicWind/Source/DynamicWind/Public/DynamicWindSubsystem.h
@@ -34,7 +34,14 @@ public:
 	DYNAMICWIND_API void UpdateWindParameters(const FDynamicWindParameters& Parameters);
 
 private:
+	void CreateTransformProvider(UWorld& World);
+	void ReleaseTransformProvider();
+
 	using FTransformProviderPtr = TSharedPtr<FDynamicWindTransformProvider, ESPMode::ThreadSafe>;
 
 	FTransformProviderPtr TransformProvider;
+#if WITH_EDITORONLY_DATA
+	FDelegateHandle OnPreRecreateScene;
+	FDelegateHandle OnPostRecreateScene;
+#endif
 };
diff --git a/Engine/Source/Runtime/Engine/Classes/Engine/World.h b/Engine/Source/Runtime/Engine/Classes/Engine/World.h
--- a/Engine/Source/Runtime/Engine/Classes/Engine/World.h
+++ b/Engine/Source/Runtime/Engine/Classes/Engine/World.h
@@ -4459,6 +4459,9 @@ public:
 	// Callback to add references when Serialize is called during SavePackage
 	DECLARE_MULTICAST_DELEGATE_TwoParams(FWorldCollectSaveReferencesEvent, UWorld*, FArchive&);
 	static UE_API FWorldCollectSaveReferencesEvent OnCollectSaveReferences;
+
+	static UE_API FWorldEvent OnPreRecreateScene;
+	static UE_API FWorldEvent OnPostRecreateScene;
 #endif // WITH_EDITOR
 
 	// Post duplication event.
diff --git a/Engine/Source/Runtime/Engine/Private/World.cpp b/Engine/Source/Runtime/Engine/Private/World.cpp
--- a/Engine/Source/Runtime/Engine/Private/World.cpp
+++ b/Engine/Source/Runtime/Engine/Private/World.cpp
@@ -782,6 +782,8 @@ FWorldDelegates::FWorldPreRenameEvent FWorldDelegates::OnPreWorldRename;
 FWorldDelegates::FWorldPostRenameEvent FWorldDelegates::OnPostWorldRename;
 FWorldDelegates::FWorldCurrentLevelChangedEvent FWorldDelegates::OnCurrentLevelChanged;
 FWorldDelegates::FWorldCollectSaveReferencesEvent FWorldDelegates::OnCollectSaveReferences;
+FWorldDelegates::FWorldEvent FWorldDelegates::OnPreRecreateScene;
+FWorldDelegates::FWorldEvent FWorldDelegates::OnPostRecreateScene;
 #endif // WITH_EDITOR
 FWorldDelegates::FWorldPostDuplicateEvent FWorldDelegates::OnPostDuplicate;
 FWorldDelegates::FWorldCleanupEvent FWorldDelegates::OnWorldCleanup;
@@ -10045,6 +10047,9 @@ void UWorld::RecreateScene(ERHIFeatureLevel::Type InFeatureLevel, bool bBroadcas
 	if (Scene)
 	{
 		ensure(InFeatureLevel == GetFeatureLevel());
+
+		FWorldDelegates::OnPreRecreateScene.Broadcast(this);
+
 		for (ULevel* Level : Levels)
 		{
 			Level->ReleaseRenderingResources();
@@ -10076,6 +10081,8 @@ void UWorld::RecreateScene(ERHIFeatureLevel::Type InFeatureLevel, bool bBroadcas
 			Level->PrecomputedVisibilityHandler.UpdateScene(Scene);
 			Level->PrecomputedVolumeDistanceField.UpdateScene(Scene);
 		}
+
+		FWorldDelegates::OnPostRecreateScene.Broadcast(this);
 	}
 }
 
