[DynamicWind]
* Fixed non unity compile error

#Synced-CL 51878269

diff --git a/Engine/Plugins/Experimental/DynamicWind/Source/DynamicWind/Private/DynamicWindProvider.cpp b/Engine/Plugins/Experimental/DynamicWind/Source/DynamicWind/Private/DynamicWindProvider.cpp
--- a/Engine/Plugins/Experimental/DynamicWind/Source/DynamicWind/Private/DynamicWindProvider.cpp
+++ b/Engine/Plugins/Experimental/DynamicWind/Source/DynamicWind/Private/DynamicWindProvider.cpp
@@ -7,6 +7,7 @@
 #include "Engine/Texture.h"
 #include "GlobalRenderResources.h"
 #include "GlobalShader.h"
+#include "NaniteSceneProxy.h"
 #include "RenderGraphBuilder.h"
 #include "RenderGraphResources.h"
 #include "RenderGraphUtils.h"
