Implement dynamic wind for GPU skin and unified bone indices which both use a bone map.

#Synced-CL 52384345

diff --git a/Engine/Plugins/Experimental/DynamicWind/Shaders/DynamicWindEval.usf b/Engine/Plugins/Experimental/DynamicWind/Shaders/DynamicWindEval.usf
--- a/Engine/Plugins/Experimental/DynamicWind/Shaders/DynamicWindEval.usf
+++ b/Engine/Plugins/Experimental/DynamicWind/Shaders/DynamicWindEval.usf
@@ -4,7 +4,7 @@
 #include "/Engine/Private/BoneTransform.ush"
 #include "/Engine/Shared/SkinningDefinitions.h"
 #include "/Engine/Private/ComputeShaderUtils.ush"
-#include "/Engine/Private/Nanite/NaniteDataDecode.ush"
+#include "/Engine/Private/Skinning/SkinningDataDecode.ush"
 
 #include "Shared/DynamicWindCommon.ush"
 
@@ -14,6 +14,7 @@
 StructuredBuffer<FDynamicWindSkeletonData> SkeletonBuffer;
 StructuredBuffer<FDynamicWindBlockHeader> HeaderBuffer;
 ByteAddressBuffer BoneDataBuffer;
+ByteAddressBuffer BoneMapBuffer;
 FCompressedBoneTransformRWBuffer TransformBuffer;
 Texture2D<float4> WindTexture;
 SamplerState WindSampler;
@@ -130,7 +131,10 @@ void WindEvalCS(uint3 GroupId : SV_GroupID, uint GroupThreadIndex : SV_GroupInde
 
 	FDynamicWindSkeletonData SkeletonData = SkeletonBuffer[Header.SkeletonDataIndex];
 
-	const uint SkeletonBoneIndex = Header.BlockBoneOffset + BoneIndexInBlock;
+	const uint TransformSlotIndex = Header.BlockBoneOffset + BoneIndexInBlock;
+	const uint SkeletonBoneIndex = (Header.BoneMapBufferOffset != ~uint(0))
+		? BoneMapBuffer.Load(Header.BoneMapBufferOffset + TransformSlotIndex * 4)
+		: TransformSlotIndex;
 	const uint BoneDataAddress = GetBoneDataAddress(SkeletonData.BoneDataOffset, SkeletonBoneIndex);
 	const FDynamicWindBoneData BoneData = BoneDataBuffer.Load<FDynamicWindBoneData>(BoneDataAddress);
 
diff --git a/Engine/Plugins/Experimental/DynamicWind/Shaders/Shared/DynamicWindCommon.ush b/Engine/Plugins/Experimental/DynamicWind/Shaders/Shared/DynamicWindCommon.ush
--- a/Engine/Plugins/Experimental/DynamicWind/Shaders/Shared/DynamicWindCommon.ush
+++ b/Engine/Plugins/Experimental/DynamicWind/Shaders/Shared/DynamicWindCommon.ush
@@ -72,6 +72,9 @@ struct FDynamicWindBlockHeader
 
 	/** Total number of transforms to be written for this skeleton anim. */
 	uint TotalTransformCount : 24;
+
+	/** Byte offset of the bone map in the bone map buffer, or ~0u if not present */
+	uint BoneMapBufferOffset;
 };
 
 #ifdef __cplusplus
diff --git a/Engine/Plugins/Experimental/DynamicWind/Source/DynamicWind/Private/DynamicWindData.cpp b/Engine/Plugins/Experimental/DynamicWind/Source/DynamicWind/Private/DynamicWindData.cpp
--- a/Engine/Plugins/Experimental/DynamicWind/Source/DynamicWind/Private/DynamicWindData.cpp
+++ b/Engine/Plugins/Experimental/DynamicWind/Source/DynamicWind/Private/DynamicWindData.cpp
@@ -13,7 +13,6 @@
 #include "SkinningSceneExtensionProxy.h"
 #include "SkinningDefinitions.h"
 #include "Components/InstancedSkinnedMeshComponent.h"
-#include "NaniteSceneProxy.h"
 
 class FSceneInterface;
 
@@ -84,8 +83,18 @@ FDynamicWindDataProxy::~FDynamicWindDataProxy()
 
 void FDynamicWindDataProxy::CreateRenderThreadResources(FSceneInterface& InScene, FRHICommandListBase& RHICmdList)
 {
+	TransformProvider = FDynamicWindTransformProvider::FindForScene(&InScene);
+	if (TransformProvider)
+	{
+		TransformProvider->RegisterProxy(SceneExtensionProxy, SkeletalData);
+	}
 }
 
 void FDynamicWindDataProxy::DestroyRenderThreadResources()
 {
+	if (TransformProvider)
+	{
+		TransformProvider->UnregisterProxy(SceneExtensionProxy);
+		TransformProvider = nullptr;
+	}
 }
\ No newline at end of file
diff --git a/Engine/Plugins/Experimental/DynamicWind/Source/DynamicWind/Private/DynamicWindProvider.cpp b/Engine/Plugins/Experimental/DynamicWind/Source/DynamicWind/Private/DynamicWindProvider.cpp
--- a/Engine/Plugins/Experimental/DynamicWind/Source/DynamicWind/Private/DynamicWindProvider.cpp
+++ b/Engine/Plugins/Experimental/DynamicWind/Source/DynamicWind/Private/DynamicWindProvider.cpp
@@ -7,7 +7,7 @@
 #include "Engine/Texture.h"
 #include "GlobalRenderResources.h"
 #include "GlobalShader.h"
-#include "NaniteSceneProxy.h"
+#include "SkinningSceneExtensionProxy.h"
 #include "RenderGraphBuilder.h"
 #include "RenderGraphResources.h"
 #include "RenderGraphUtils.h"
@@ -90,7 +90,7 @@ private:
 
 	static bool ShouldCompilePermutation(const FGlobalShaderPermutationParameters& Parameters)
 	{
-		return DoesPlatformSupportNanite(Parameters.Platform);
+		return true;
 	}
 
 	static void ModifyCompilationEnvironment(const FGlobalShaderPermutationParameters& Parameters, FShaderCompilerEnvironment& OutEnvironment)
@@ -99,6 +99,7 @@ private:
 
 		OutEnvironment.CompilerFlags.Add(CFLAG_WarningsAsErrors);
 		OutEnvironment.CompilerFlags.Add(CFLAG_HLSL2021);
+		OutEnvironment.CompilerFlags.Add(CFLAG_AllowTypedUAVLoads);
 
 		OutEnvironment.SetDefine(TEXT("BONES_PER_GROUP"), BonesPerGroup);
 	}
@@ -107,6 +108,7 @@ private:
 		SHADER_PARAMETER_RDG_BUFFER_SRV(StructuredBuffer<FDynamicWindSkeletonData>, SkeletonBuffer)
 		SHADER_PARAMETER_RDG_BUFFER_SRV(StructuredBuffer<FDynamicWindBlockHeader>, HeaderBuffer)
 		SHADER_PARAMETER_RDG_BUFFER_SRV(ByteAddressBuffer, BoneDataBuffer)
+		SHADER_PARAMETER_RDG_BUFFER_SRV(ByteAddressBuffer, BoneMapBuffer)
 		SHADER_PARAMETER_RDG_COMPRESSED_BONE_TRANSFORM_UAV(TransformBuffer)
 		SHADER_PARAMETER_RDG_TEXTURE_SRV(Texture2D<float4>, WindTexture)
 		SHADER_PARAMETER_SAMPLER(SamplerState, WindSampler)
@@ -121,6 +123,20 @@ private:
 
 IMPLEMENT_GLOBAL_SHADER(FDynamicWindEvalCS, "/Plugin/DynamicWind/DynamicWindEval.usf", "WindEvalCS", SF_Compute);
 
+TMap<const FScene*, FDynamicWindTransformProvider*> FDynamicWindTransformProvider::SceneProviderMap;
+
+FDynamicWindTransformProvider* FDynamicWindTransformProvider::FindForScene(const FSceneInterface* InScene)
+{
+	if (const FScene* RenderScene = InScene ? InScene->GetRenderScene() : nullptr)
+	{
+		if (auto* Found = SceneProviderMap.Find(RenderScene))
+		{
+			return *Found;
+		}
+	}
+	return nullptr;
+}
+
 FDynamicWindTransformProvider::FDynamicWindTransformProvider(FScene& InScene)
 : Scene(InScene)
 , BoneDataBuffer(1, TEXT("DynamicWind.BoneDataBuffer"))
@@ -144,31 +160,7 @@ bool FDynamicWindTransformProvider::Register()
 			true /* Use skeleton batching */
 		);
 
-		OnProxyCreateRenderThreadResources =
-			Nanite::FSkinnedSceneProxyDelegates::OnCreateRenderThreadResources.AddLambda(
-				[this](const Nanite::FSkinnedSceneProxy* InSkinnedSceneProxy)
-				{
-					if (&InSkinnedSceneProxy->GetScene() == &Scene)
-					{
-						PRAGMA_DISABLE_DEPRECATION_WARNINGS
-						RegisterSceneProxy(InSkinnedSceneProxy);
-						PRAGMA_ENABLE_DEPRECATION_WARNINGS
-					}
-				}
-			);
-
-		OnProxyDestroyRenderThreadResources =
-			Nanite::FSkinnedSceneProxyDelegates::OnDestroyRenderThreadResources.AddLambda(
-				[this](const Nanite::FSkinnedSceneProxy* InSkinnedSceneProxy)
-				{
-					if (&InSkinnedSceneProxy->GetScene() == &Scene)
-					{
-						PRAGMA_DISABLE_DEPRECATION_WARNINGS
-						UnregisterSceneProxy(InSkinnedSceneProxy);
-						PRAGMA_ENABLE_DEPRECATION_WARNINGS
-					}
-				}
-			);
+		SceneProviderMap.Add(&Scene, this);
 
 		bIsRegistered = true;
 		return true;
@@ -186,17 +178,7 @@ bool FDynamicWindTransformProvider::Unregister()
 			TransformProvider->UnregisterProvider(DynamicWindTransformProviderId);
 		}
 
-		if (OnProxyCreateRenderThreadResources.IsValid())
-		{
-			Nanite::FSkinnedSceneProxyDelegates::OnCreateRenderThreadResources.Remove(OnProxyCreateRenderThreadResources);
-			OnProxyCreateRenderThreadResources.Reset();
-		}
-
-		if (OnProxyDestroyRenderThreadResources.IsValid())
-		{
-			Nanite::FSkinnedSceneProxyDelegates::OnDestroyRenderThreadResources.Remove(OnProxyDestroyRenderThreadResources);
-			OnProxyDestroyRenderThreadResources.Reset();
-		}
+		SceneProviderMap.Remove(&Scene);
 
 		bIsRegistered = false;
 		return true;
@@ -204,22 +186,16 @@ bool FDynamicWindTransformProvider::Unregister()
 
 	return false;
 }
-	
-void FDynamicWindTransformProvider::RegisterSceneProxy(const Nanite::FSkinnedSceneProxy* InProxy)
+
+void FDynamicWindTransformProvider::RegisterProxy(const FSkinningSceneExtensionProxy* ExtensionProxy, const UDynamicWindSkeletalData* InSkeletalData)
 {
 	if (!bIsRegistered)
 	{
-		UE_LOG(LogDynamicWind, Verbose, TEXT("Trying to register SkinnedSceneProxy (SkinnedAsset=%s), when the provider is not registered."), *GetNameSafe(InProxy->GetSkinnedAsset()));
+		UE_LOG(LogDynamicWind, Verbose, TEXT("Trying to register proxy (SkinnedAsset=%s), when the provider is not registered."), *GetNameSafe(ExtensionProxy->GetSkinnedAsset()));
 		return;
 	}
 
-	const FSkinningSceneExtensionProxy* SkinningSceneExtProxy = InProxy->GetSkinningSceneExtensionProxy();
-	if (SkinningSceneExtProxy == nullptr || SkinningSceneExtProxy->GetTransformProviderId() != DynamicWindTransformProviderId)
-	{
-		return;
-	}
-
-	USkeletalMesh* SkeletalMesh = const_cast<USkeletalMesh*>(Cast<USkeletalMesh>(InProxy->GetSkinnedAsset()));
+	const USkeletalMesh* SkeletalMesh = Cast<USkeletalMesh>(ExtensionProxy->GetSkinnedAsset());
 	if (SkeletalMesh == nullptr)
 	{
 		return;
@@ -228,8 +204,8 @@ void FDynamicWindTransformProvider::RegisterSceneProxy(const Nanite::FSkinnedSce
 	const USkeleton* Skeleton = SkeletalMesh->GetSkeleton();
 	const FReferenceSkeleton& RefSkeleton = Skeleton->GetReferenceSkeleton();
 	const uint32 NumBonesInSkeleton = RefSkeleton.GetNum();
-	
-	const UDynamicWindSkeletalData* Data = SkeletalMesh->GetAssetUserData<UDynamicWindSkeletalData>();
+
+	const UDynamicWindSkeletalData* Data = InSkeletalData;
 	if (Data == nullptr)
 	{
 		// No user data, don't register the skeleton. The wind simulation CS will provide identity
@@ -237,8 +213,9 @@ void FDynamicWindTransformProvider::RegisterSceneProxy(const Nanite::FSkinnedSce
 	}
 
 	const uint64 UserDataHash = Data->GetSkeletalDataHash();
-	
-	FSkeletonEntry& SkeletonEntry = SkeletonLookup.FindOrAdd(Skeleton->GetGuid());
+
+	const bool bUsesBoneMap = ExtensionProxy->GetBoneTransformStorageMode() == EBoneTransformStorageMode::BoneMap;
+	FSkeletonEntry& SkeletonEntry = SkeletonLookup.FindOrAdd(FSkeletonKey{ Skeleton->GetGuid(), bUsesBoneMap });
 
 	++SkeletonEntry.ReferenceCount;
 	if (SkeletonEntry.ReferenceCount == 1)
@@ -300,23 +277,23 @@ void FDynamicWindTransformProvider::RegisterSceneProxy(const Nanite::FSkinnedSce
 	}
 }
 
-void FDynamicWindTransformProvider::UnregisterSceneProxy(const Nanite::FSkinnedSceneProxy* InProxy)
+void FDynamicWindTransformProvider::UnregisterProxy(const FSkinningSceneExtensionProxy* ExtensionProxy)
 {
 	if (!bIsRegistered)
 	{
-		UE_LOG(LogDynamicWind, Verbose, TEXT("Trying to unregister SkinnedSceneProxy (SkinnedAsset=%s), when the provider is not registered."), *GetNameSafe(InProxy->GetSkinnedAsset()));
+		UE_LOG(LogDynamicWind, Verbose, TEXT("Trying to unregister proxy (SkinnedAsset=%s), when the provider is not registered."), *GetNameSafe(ExtensionProxy->GetSkinnedAsset()));
 		return;
 	}
 
-	const FSkinningSceneExtensionProxy* SkinningSceneExtProxy = InProxy->GetSkinningSceneExtensionProxy();
-	if (SkinningSceneExtProxy == nullptr || SkinningSceneExtProxy->GetTransformProviderId() != DynamicWindTransformProviderId)
+	const USkeletalMesh* SkeletalMesh = Cast<USkeletalMesh>(ExtensionProxy->GetSkinnedAsset());
+	if (SkeletalMesh == nullptr)
 	{
 		return;
 	}
 
-	const USkeletalMesh* SkeletalMesh = Cast<USkeletalMesh>(InProxy->GetSkinnedAsset());
 	const USkeleton* Skeleton = SkeletalMesh->GetSkeleton();
-	auto EntryID = SkeletonLookup.FindId(Skeleton->GetGuid());
+	const bool bUsesBoneMap = ExtensionProxy->GetBoneTransformStorageMode() == EBoneTransformStorageMode::BoneMap;
+	auto EntryID = SkeletonLookup.FindId(FSkeletonKey{ Skeleton->GetGuid(), bUsesBoneMap });
 
 	if (!SkeletonLookup.IsValidId(EntryID))
 	{
@@ -381,7 +358,7 @@ void FDynamicWindTransformProvider::ProvideTransforms(FSkinningTransformProvider
 		[
 			BlockHeaders,
 			SkeletonData,
-			SkeletonLookup = TMap<FGuid, FSkeletonEntry, SceneRenderingSetAllocator>(SkeletonLookup), // Copy for task
+			SkeletonLookup = TMap<FSkeletonKey, FSkeletonEntry, SceneRenderingSetAllocator>(SkeletonLookup), // Copy for task
 			Indirections = Context.Indirections,
 			SkeletonBatches = Context.SkeletonBatches
 		]
@@ -408,7 +385,10 @@ void FDynamicWindTransformProvider::ProvideTransforms(FSkinningTransformProvider
 
 				uint32 DstTransformOffset = Indirection.CurrentTransformOffset;
 
-				FSetElementId SkeletonId = SkeletonLookup.FindId(SkinnedAnimationBatch.SkeletonGuid);
+				const bool bUsesBoneMap = Indirection.BoneMapOffset != INDEX_NONE;
+				FSetElementId SkeletonId = SkeletonLookup.FindId(FSkeletonKey{ SkinnedAnimationBatch.SkeletonGuid, bUsesBoneMap });
+
+				const uint32 BoneMapBufferOffset = Indirection.BoneMapOffset;
 
 				for (uint32 DirectionalityIndex = 0; DirectionalityIndex < UniqueAnimationCount; ++DirectionalityIndex)
 				{
@@ -422,6 +402,7 @@ void FDynamicWindTransformProvider::ProvideTransforms(FSkinningTransformProvider
 						BlockHeaders[BlockWrite].BlockBoneOffset = BoneOffset;
 						BlockHeaders[BlockWrite].BlockDirectionalityIndex = DirectionalityIndex;
 						BlockHeaders[BlockWrite].TotalTransformCount = TransformCount;
+						BlockHeaders[BlockWrite].BoneMapBufferOffset = BoneMapBufferOffset;
 						
 						BoneOffset += BonesPerGroup;
 						DstTransformOffset += (BonesPerGroup * sizeof(FCompressedBoneTransform));
@@ -436,6 +417,7 @@ void FDynamicWindTransformProvider::ProvideTransforms(FSkinningTransformProvider
 						BlockHeaders[BlockWrite].BlockBoneOffset = BoneOffset;
 						BlockHeaders[BlockWrite].BlockDirectionalityIndex = DirectionalityIndex;
 						BlockHeaders[BlockWrite].TotalTransformCount = TransformCount;
+						BlockHeaders[BlockWrite].BoneMapBufferOffset = BoneMapBufferOffset;
 
 						BoneOffset += PartialTransformCount;
 						DstTransformOffset += (PartialTransformCount * sizeof(FCompressedBoneTransform));
@@ -524,6 +506,7 @@ void FDynamicWindTransformProvider::ProvideTransforms(FSkinningTransformProvider
 		PassParameters->SkeletonBuffer = GraphBuilder.CreateSRV(SkeletonBuffer);
 		PassParameters->HeaderBuffer = GraphBuilder.CreateSRV(BlockHeaderBuffer);
 		PassParameters->BoneDataBuffer = GraphBuilder.CreateSRV(BoneDataBufferRDG);
+		PassParameters->BoneMapBuffer = Context.BoneMapBufferSRV;
 		PassParameters->TransformBuffer = GetCompressedBoneTransformUAV(GraphBuilder, Context.TransformBuffer);
 		PassParameters->WindTexture = GraphBuilder.CreateSRV(WindTextureRef);
 		PassParameters->WindSampler = WindSampler;
diff --git a/Engine/Plugins/Experimental/DynamicWind/Source/DynamicWind/Private/DynamicWindProvider.h b/Engine/Plugins/Experimental/DynamicWind/Source/DynamicWind/Private/DynamicWindProvider.h
--- a/Engine/Plugins/Experimental/DynamicWind/Source/DynamicWind/Private/DynamicWindProvider.h
+++ b/Engine/Plugins/Experimental/DynamicWind/Source/DynamicWind/Private/DynamicWindProvider.h
@@ -12,7 +12,8 @@
 #include "Animation/TransformProviderData.h"
 
 class USkeleton;
-namespace Nanite { class FSkinnedSceneProxy; }
+class FSceneInterface;
+class UDynamicWindSkeletalData;
 
 #define DYNAMIC_WIND_TRANSFORM_PROVIDER_GUID 0xDFD4874B, 0xEE57466D, 0x883D6419, 0xAFE99EAC
 #define DYNAMIC_WIND_DEBUG_SKELETON_NAMES (UE_BUILD_DEBUG | UE_BUILD_DEVELOPMENT)
@@ -27,6 +28,10 @@ public:
 	bool Register();
 	bool Unregister();
 
+	static FDynamicWindTransformProvider* FindForScene(const FSceneInterface* Scene);
+
+	void RegisterProxy(const FSkinningSceneExtensionProxy* ExtensionProxy, const UDynamicWindSkeletalData* InSkeletalData);
+	void UnregisterProxy(const FSkinningSceneExtensionProxy* ExtensionProxy);
 	int32 UpdateParameters(const FDynamicWindParameters& Parameters);
 
 	float GetBlendedWindAmplitude() const;
@@ -35,6 +40,22 @@ private:
 	using FBoneDataBuffer = TPersistentByteAddressBuffer<FDynamicWindBoneData>;
 	using FBoneDataUploader = TByteAddressBufferScatterUploader<FDynamicWindBoneData>;
 
+	struct FSkeletonKey
+	{
+		FGuid SkeletonGuid;
+		bool bUsesBoneMap = false;
+
+		bool operator==(const FSkeletonKey& Other) const
+		{
+			return SkeletonGuid == Other.SkeletonGuid && bUsesBoneMap == Other.bUsesBoneMap;
+		}
+
+		friend uint32 GetTypeHash(const FSkeletonKey& Key)
+		{
+			return HashCombine(GetTypeHash(Key.SkeletonGuid), GetTypeHash(Key.bUsesBoneMap));
+		}
+	};
+
 	struct FSkeletonEntry
 	{	
 		uint64 UserDataHash = 0;
@@ -47,14 +68,12 @@ private:
 		FDynamicWindSkeletonData Data;
 	};
 
-	void RegisterSceneProxy(const Nanite::FSkinnedSceneProxy* InProxy);
-	void UnregisterSceneProxy(const Nanite::FSkinnedSceneProxy* InProxy);
 	void ProvideTransforms(FSkinningTransformProvider::FProviderContext& Context);
 
 private:
 	FScene& Scene;
 
-	TMap<FGuid, FSkeletonEntry> SkeletonLookup;
+	TMap<FSkeletonKey, FSkeletonEntry> SkeletonLookup;
 	FSpanAllocator BoneDataAllocator;
 	TUniquePtr<FBoneDataUploader> BoneDataUploader;
 	FBoneDataBuffer BoneDataBuffer;
@@ -65,8 +84,7 @@ private:
 
 	float BlendedWindAmplitude = -1.0f;
 
-	FDelegateHandle OnProxyCreateRenderThreadResources;
-	FDelegateHandle OnProxyDestroyRenderThreadResources;
-
 	bool bIsRegistered = false;
+
+	static TMap<const FScene*, FDynamicWindTransformProvider*> SceneProviderMap;
 };
\ No newline at end of file
diff --git a/Engine/Plugins/Experimental/DynamicWind/Source/DynamicWind/Private/DynamicWindSubsystem.cpp b/Engine/Plugins/Experimental/DynamicWind/Source/DynamicWind/Private/DynamicWindSubsystem.cpp
--- a/Engine/Plugins/Experimental/DynamicWind/Source/DynamicWind/Private/DynamicWindSubsystem.cpp
+++ b/Engine/Plugins/Experimental/DynamicWind/Source/DynamicWind/Private/DynamicWindSubsystem.cpp
@@ -15,7 +15,6 @@
 #include "SceneView.h"
 #include "SceneInterface.h"
 #include "ScenePrivate.h"
-#include "NaniteSceneProxy.h"
 #include "DynamicWindSkeletalData.h"
 #include "DynamicWindProvider.h"
 
diff --git a/Engine/Plugins/Experimental/DynamicWind/Source/DynamicWind/Public/DynamicWindData.h b/Engine/Plugins/Experimental/DynamicWind/Source/DynamicWind/Public/DynamicWindData.h
--- a/Engine/Plugins/Experimental/DynamicWind/Source/DynamicWind/Public/DynamicWindData.h
+++ b/Engine/Plugins/Experimental/DynamicWind/Source/DynamicWind/Public/DynamicWindData.h
@@ -35,4 +35,5 @@ public:
 private:
 	FInstancedSkinningSceneExtensionProxy* SceneExtensionProxy = nullptr;
 	TObjectPtr<const class UDynamicWindSkeletalData> SkeletalData = nullptr;
+	class FDynamicWindTransformProvider* TransformProvider = nullptr;
 };
diff --git a/Engine/Plugins/Experimental/DynamicWind/Source/DynamicWind/Public/DynamicWindSubsystem.h b/Engine/Plugins/Experimental/DynamicWind/Source/DynamicWind/Public/DynamicWindSubsystem.h
--- a/Engine/Plugins/Experimental/DynamicWind/Source/DynamicWind/Public/DynamicWindSubsystem.h
+++ b/Engine/Plugins/Experimental/DynamicWind/Source/DynamicWind/Public/DynamicWindSubsystem.h
@@ -13,7 +13,6 @@
 class USkeleton;
 class UInstancedSkinnedMeshComponent;
 class FDynamicWindTransformProvider;
-namespace Nanite { class FSkinnedSceneProxy; }
 
 UCLASS(BlueprintType, ClassGroup = (Rendering, Common))
 class UDynamicWindSubsystem : public UWorldSubsystem
