Fixed dynamic wind ES31 compilation error not supporting bit fields in structured buffers.

#Synced-CL 52687702

diff --git a/Engine/Plugins/Experimental/DynamicWind/Shaders/DynamicWindEval.usf b/Engine/Plugins/Experimental/DynamicWind/Shaders/DynamicWindEval.usf
--- a/Engine/Plugins/Experimental/DynamicWind/Shaders/DynamicWindEval.usf
+++ b/Engine/Plugins/Experimental/DynamicWind/Shaders/DynamicWindEval.usf
@@ -109,7 +109,12 @@ void WindEvalCS(uint3 GroupId : SV_GroupID, uint GroupThreadIndex : SV_GroupInde
 	const uint BoneIndexInBlock	= GroupThreadIndex;
 
 	const FDynamicWindBlockHeader Header = HeaderBuffer[GlobalBlockIndex];
-	if (BoneIndexInBlock >= Header.BlockNumBones)
+	const uint BlockNumBones = GetBlockNumBones(Header);
+	const uint BlockBoneOffset = GetBlockBoneOffset(Header);
+	const uint BlockDirectionalityIndex = GetBlockDirectionalityIndex(Header);
+	const uint TotalTransformCount = GetTotalTransformCount(Header);
+
+	if (BoneIndexInBlock >= BlockNumBones)
 	{
 		return;
 	}
@@ -124,14 +129,16 @@ void WindEvalCS(uint3 GroupId : SV_GroupID, uint GroupThreadIndex : SV_GroupInde
 			float3(0, 0, 0)
 		};
 		StoreCompressedBoneTransform(TransformBuffer, Header.BlockDstTransformOffset, BoneIndexInBlock, Identity);
-		StoreCompressedBoneTransform(TransformBuffer, Header.BlockDstTransformOffset, BoneIndexInBlock + Header.TotalTransformCount, Identity);
+		StoreCompressedBoneTransform(TransformBuffer, Header.BlockDstTransformOffset, BoneIndexInBlock + TotalTransformCount, Identity);
 
 		return;
 	}
 
 	FDynamicWindSkeletonData SkeletonData = SkeletonBuffer[Header.SkeletonDataIndex];
+	const uint MaxSimulationGroupIndex = GetMaxSimulationGroupIndex(SkeletonData);
+	const uint IsGroundCover = GetIsGroundCover(SkeletonData);
 
-	const uint TransformSlotIndex = Header.BlockBoneOffset + BoneIndexInBlock;
+	const uint TransformSlotIndex = BlockBoneOffset + BoneIndexInBlock;
 	const uint SkeletonBoneIndex = (Header.BoneMapBufferOffset != ~uint(0))
 		? BoneMapBuffer.Load(Header.BoneMapBufferOffset + TransformSlotIndex * 4)
 		: TransformSlotIndex;
@@ -154,7 +161,7 @@ void WindEvalCS(uint3 GroupId : SV_GroupID, uint GroupThreadIndex : SV_GroupInde
 		const FDynamicWindBoneData CurrentBoneData = BoneDataBuffer.Load<FDynamicWindBoneData>(CurrentBoneDataAddress);
 
 		// === ROTATION AXIS CALCULATIONS ========================================================================================================
-		const float3 SectionWindDirection = QuatRotateVector(QuatConjugate(QuatFromAxisAngle(float3(0.0f, 0.0f, 1.0f), Header.BlockDirectionalityIndex * SliceAngle)), WindDirection);
+		const float3 SectionWindDirection = QuatRotateVector(QuatConjugate(QuatFromAxisAngle(float3(0.0f, 0.0f, 1.0f), BlockDirectionalityIndex * SliceAngle)), WindDirection);
 		const float3 AdjustedBoneForward = QuatRotateVector(CurrentBoneData.BindPoseRotation, float3(1.0f, 0.0f, 0.0f));
 			
 		// Prevent error when the bone forward and wind direction are collinear
@@ -170,7 +177,7 @@ void WindEvalCS(uint3 GroupId : SV_GroupID, uint GroupThreadIndex : SV_GroupInde
 		const float HeightGrad = saturate(CurrentBoneData.BindPosePosition.z / SkeletonData.SkeletonHeight);
 		// RandBBSfloat(0) returns 0, which result in an uniform rotation for all branches
 		// SkeletonDataIndex is 0 for the first registered skeleton, hence the +1
-		const float RandomValue = RandBBSfloat(((Header.SkeletonDataIndex << 8) | Header.BlockDirectionalityIndex) + 1);
+		const float RandomValue = RandBBSfloat(((Header.SkeletonDataIndex << 8) | BlockDirectionalityIndex) + 1);
 			
 		float GustSpeed = lerp(DYNAMIC_WIND_GUST_SPEED_MIN, DYNAMIC_WIND_GUST_SPEED_MAX, WindSpeed / 100.0f);
 		float GustStrength = RampStrength(WindSpeed / 100.0f);
@@ -201,7 +208,7 @@ void WindEvalCS(uint3 GroupId : SV_GroupID, uint GroupThreadIndex : SV_GroupInde
 		else // Branches
 		{
 			const float DistanceFromOrigin = distance(CurrentBoneData.BindPosePosition, float3(0.0f, 0.0f, 0.0f));
-			const float BlendByGroup = max(CurrentBoneData.SimulationGroupIndex, 0) / (SkeletonData.MaxSimulationGroupIndex + 1);
+			const float BlendByGroup = max(CurrentBoneData.SimulationGroupIndex, 0) / (MaxSimulationGroupIndex + 1);
 			const float TrunkBlend = (float)CurrentBoneData.bIsTrunkBone;
 				
 			float Gust = GustStrength * max(0, WindWave(GustSpeed * (Time + 0.5f) + RandomValue * 17));
@@ -231,7 +238,7 @@ void WindEvalCS(uint3 GroupId : SV_GroupID, uint GroupThreadIndex : SV_GroupInde
 			// Attenuate branch final rotation on lower branches
 			// It's here and not on the branch section because I want to influence all the rotation on the branch
 			const float RawHeightBlend = CurrentBoneData.BindPosePosition.z / SkeletonData.SkeletonHeight;
-			const float FinalHeightBlend = max(SkeletonData.bIsGroundCover, RemapValue(RawHeightBlend, 0.0f, 1.0f, DYNAMIC_WIND_HEIGHT_ATTENUATION, 1.0f));
+			const float FinalHeightBlend = max(IsGroundCover, RemapValue(RawHeightBlend, 0.0f, 1.0f, DYNAMIC_WIND_HEIGHT_ATTENUATION, 1.0f));
 			PartialRotation = QuatAttenuate(PartialRotation, FinalHeightBlend * WindAmplitude);
 		}
 			
@@ -254,5 +261,5 @@ void WindEvalCS(uint3 GroupId : SV_GroupID, uint GroupThreadIndex : SV_GroupInde
 	// TODO: Handle first frame by writing CurrentTransform for both
 
 	StoreCompressedBoneTransform(TransformBuffer, Header.BlockDstTransformOffset, BoneIndexInBlock, CurrentTransform);
-	StoreCompressedBoneTransform(TransformBuffer, Header.BlockDstTransformOffset, BoneIndexInBlock + Header.TotalTransformCount, PreviousTransform);
+	StoreCompressedBoneTransform(TransformBuffer, Header.BlockDstTransformOffset, BoneIndexInBlock + TotalTransformCount, PreviousTransform);
 }
\ No newline at end of file
diff --git a/Engine/Plugins/Experimental/DynamicWind/Shaders/Shared/DynamicWindCommon.ush b/Engine/Plugins/Experimental/DynamicWind/Shaders/Shared/DynamicWindCommon.ush
--- a/Engine/Plugins/Experimental/DynamicWind/Shaders/Shared/DynamicWindCommon.ush
+++ b/Engine/Plugins/Experimental/DynamicWind/Shaders/Shared/DynamicWindCommon.ush
@@ -22,11 +22,8 @@ struct FDynamicWindSkeletonData
 	/** How much to attenuate trunk bones in the skeleton (0-1) */
 	float GustAttenuation;
 
-	/** The maximum SimulationGroupIndex value for all bones in the skeleton */
-	uint MaxSimulationGroupIndex : 31;
-
-	/** Whether or not this object is ground cover */
-	uint bIsGroundCover : 1;
+	/** Packed: MaxSimulationGroupIndex (bits 0-30), bIsGroundCover (bit 31) */
+	uint PackedSimGroupAndGroundCover;
 };
 
 struct FDynamicWindBoneData
@@ -61,17 +58,11 @@ struct FDynamicWindBlockHeader
 	/** Byte offset of the transforms to be written in the transform buffer */
 	uint BlockDstTransformOffset;
 
-	/** Number of bones in this block */
-	uint BlockNumBones : 8;
-
-	/** Index of the first bone (in the skeleton) of the block */
-	uint BlockBoneOffset : 24;
-
-	/** Index of the unique directional anim this block is associated with */
-	uint BlockDirectionalityIndex : 8;
+	/** Packed: BlockNumBones (bits 0-7), BlockBoneOffset (bits 8-31) */
+	uint PackedNumBonesAndOffset;
 
-	/** Total number of transforms to be written for this skeleton anim. */
-	uint TotalTransformCount : 24;
+	/** Packed: BlockDirectionalityIndex (bits 0-7), TotalTransformCount (bits 8-31) */
+	uint PackedDirectionalityAndTransformCount;
 
 	/** Byte offset of the bone map in the bone map buffer, or ~0u if not present */
 	uint BoneMapBufferOffset;
@@ -84,4 +75,25 @@ using FDynamicWindSkeletonData = UE::HLSL::FDynamicWindSkeletonData;
 using FDynamicWindBoneData = UE::HLSL::FDynamicWindBoneData;
 using FDynamicWindBlockHeader = UE::HLSL::FDynamicWindBlockHeader;
 
+// C++ setters for packed fields
+
+inline void SetMaxSimulationGroupIndex(FDynamicWindSkeletonData& S, uint32 V) { S.PackedSimGroupAndGroundCover = (S.PackedSimGroupAndGroundCover & 0x80000000u) | (V & 0x7FFFFFFFu); }
+inline void SetIsGroundCover(FDynamicWindSkeletonData& S, uint32 V) { S.PackedSimGroupAndGroundCover = (S.PackedSimGroupAndGroundCover & 0x7FFFFFFFu) | (V << 31); }
+
+inline void SetBlockNumBones(FDynamicWindBlockHeader& H, uint32 V) { H.PackedNumBonesAndOffset = (H.PackedNumBonesAndOffset & 0xFFFFFF00u) | (V & 0xFFu); }
+inline void SetBlockBoneOffset(FDynamicWindBlockHeader& H, uint32 V) { H.PackedNumBonesAndOffset = (H.PackedNumBonesAndOffset & 0xFFu) | (V << 8); }
+
+inline void SetBlockDirectionalityIndex(FDynamicWindBlockHeader& H, uint32 V) { H.PackedDirectionalityAndTransformCount = (H.PackedDirectionalityAndTransformCount & 0xFFFFFF00u) | (V & 0xFFu); }
+inline void SetTotalTransformCount(FDynamicWindBlockHeader& H, uint32 V) { H.PackedDirectionalityAndTransformCount = (H.PackedDirectionalityAndTransformCount & 0xFFu) | (V << 8); }
+
+#else // HLSL accessors for packed fields
+
+uint GetMaxSimulationGroupIndex(FDynamicWindSkeletonData S) { return S.PackedSimGroupAndGroundCover & 0x7FFFFFFFu; }
+uint GetIsGroundCover(FDynamicWindSkeletonData S) { return S.PackedSimGroupAndGroundCover >> 31; }
+
+uint GetBlockNumBones(FDynamicWindBlockHeader H) { return H.PackedNumBonesAndOffset & 0xFFu; }
+uint GetBlockBoneOffset(FDynamicWindBlockHeader H) { return H.PackedNumBonesAndOffset >> 8; }
+uint GetBlockDirectionalityIndex(FDynamicWindBlockHeader H) { return H.PackedDirectionalityAndTransformCount & 0xFFu; }
+uint GetTotalTransformCount(FDynamicWindBlockHeader H) { return H.PackedDirectionalityAndTransformCount >> 8; }
+
 #endif //__cplusplus
diff --git a/Engine/Plugins/Experimental/DynamicWind/Source/DynamicWind/Private/DynamicWindProvider.cpp b/Engine/Plugins/Experimental/DynamicWind/Source/DynamicWind/Private/DynamicWindProvider.cpp
--- a/Engine/Plugins/Experimental/DynamicWind/Source/DynamicWind/Private/DynamicWindProvider.cpp
+++ b/Engine/Plugins/Experimental/DynamicWind/Source/DynamicWind/Private/DynamicWindProvider.cpp
@@ -245,10 +245,11 @@ void FDynamicWindTransformProvider::RegisterProxy(const FSkinningSceneExtensionP
 	const FBoxSphereBounds Bounds = SkeletalMesh->GetBounds();
 
 	SkeletonEntry.UserDataHash = UserDataHash;
-	SkeletonEntry.Data.MaxSimulationGroupIndex = uint32(FMath::Max(Data->SimulationGroups.Num() - 1, 0));
+	SkeletonEntry.Data.PackedSimGroupAndGroundCover = 0;
+	SetMaxSimulationGroupIndex(SkeletonEntry.Data, uint32(FMath::Max(Data->SimulationGroups.Num() - 1, 0)));
+	SetIsGroundCover(SkeletonEntry.Data, Data->bIsGroundCover ? 1u : 0u);
 	SkeletonEntry.Data.SkeletonHeight = Bounds.BoxExtent.Z + Bounds.Origin.Z;
 	SkeletonEntry.Data.GustAttenuation = Data->GustAttenuation;
-	SkeletonEntry.Data.bIsGroundCover = Data->bIsGroundCover;
 
 	TArrayView<FDynamicWindBoneData> BoneData = BoneDataUploader->AddMultiple_GetRef(SkeletonEntry.Data.BoneDataOffset, NumBonesInSkeleton);
 	for (uint32 BoneIndex = 0; BoneIndex < NumBonesInSkeleton; ++BoneIndex)
@@ -396,13 +397,16 @@ void FDynamicWindTransformProvider::ProvideTransforms(FSkinningTransformProvider
 
 					for (uint32 BlockIndex = 0; BlockIndex < FullBlockCount; ++BlockIndex)
 					{
-						BlockHeaders[BlockWrite].SkeletonDataIndex = SkeletonId.AsInteger(); // NOTE: Can be INDEX_NONE
-						BlockHeaders[BlockWrite].BlockDstTransformOffset = DstTransformOffset;
-						BlockHeaders[BlockWrite].BlockNumBones = BonesPerGroup;
-						BlockHeaders[BlockWrite].BlockBoneOffset = BoneOffset;
-						BlockHeaders[BlockWrite].BlockDirectionalityIndex = DirectionalityIndex;
-						BlockHeaders[BlockWrite].TotalTransformCount = TransformCount;
-						BlockHeaders[BlockWrite].BoneMapBufferOffset = BoneMapBufferOffset;
+						FDynamicWindBlockHeader& Header = BlockHeaders[BlockWrite];
+						Header.SkeletonDataIndex = SkeletonId.AsInteger(); // NOTE: Can be INDEX_NONE
+						Header.BlockDstTransformOffset = DstTransformOffset;
+						Header.PackedNumBonesAndOffset = 0;
+						SetBlockNumBones(Header, BonesPerGroup);
+						SetBlockBoneOffset(Header, BoneOffset);
+						Header.PackedDirectionalityAndTransformCount = 0;
+						SetBlockDirectionalityIndex(Header, DirectionalityIndex);
+						SetTotalTransformCount(Header, TransformCount);
+						Header.BoneMapBufferOffset = BoneMapBufferOffset;
 						
 						BoneOffset += BonesPerGroup;
 						DstTransformOffset += (BonesPerGroup * sizeof(FCompressedBoneTransform));
@@ -411,13 +415,16 @@ void FDynamicWindTransformProvider::ProvideTransforms(FSkinningTransformProvider
 
 					if (PartialTransformCount > 0)
 					{
-						BlockHeaders[BlockWrite].SkeletonDataIndex = SkeletonId.AsInteger(); // NOTE: Can be INDEX_NONE
-						BlockHeaders[BlockWrite].BlockDstTransformOffset = DstTransformOffset;
-						BlockHeaders[BlockWrite].BlockNumBones = PartialTransformCount;
-						BlockHeaders[BlockWrite].BlockBoneOffset = BoneOffset;
-						BlockHeaders[BlockWrite].BlockDirectionalityIndex = DirectionalityIndex;
-						BlockHeaders[BlockWrite].TotalTransformCount = TransformCount;
-						BlockHeaders[BlockWrite].BoneMapBufferOffset = BoneMapBufferOffset;
+						FDynamicWindBlockHeader& Header = BlockHeaders[BlockWrite];
+						Header.SkeletonDataIndex = SkeletonId.AsInteger(); // NOTE: Can be INDEX_NONE
+						Header.BlockDstTransformOffset = DstTransformOffset;
+						Header.PackedNumBonesAndOffset = 0;
+						SetBlockNumBones(Header, PartialTransformCount);
+						SetBlockBoneOffset(Header, BoneOffset);
+						Header.PackedDirectionalityAndTransformCount = 0;
+						SetBlockDirectionalityIndex(Header, DirectionalityIndex);
+						SetTotalTransformCount(Header, TransformCount);
+						Header.BoneMapBufferOffset = BoneMapBufferOffset;
 
 						BoneOffset += PartialTransformCount;
 						DstTransformOffset += (PartialTransformCount * sizeof(FCompressedBoneTransform));
